Translation Editor:

- Fixed some crashes that occured if certain tabs were closed.
- Fixed issue where connection to source control would fail in some cases, but warning popup would not display
- Fixed bug where Completed, Needs Review, and Untranslated tabs could not be re-opened if closed.

diff --git a/Engine/Source/Editor/TranslationEditor/Private/TranslationDataManager.cpp b/Engine/Source/Editor/TranslationEditor/Private/TranslationDataManager.cpp
--- a/Engine/Source/Editor/TranslationEditor/Private/TranslationDataManager.cpp
+++ b/Engine/Source/Editor/TranslationEditor/Private/TranslationDataManager.cpp
@@ -25,11 +25,6 @@ FTranslationDataManager::FTranslationDataManager( const FString& ManifestFile, c
 	TranslationData = NewObject<UTranslationDataObject>();
 	check (TranslationData != NULL);
 
-	// Used to use this to trick the AssetEditorToolkit into giving us save buttons, 
-	// but it also added "Find in Content Browser" buttons that don't make sense at this time for Translation Editor
-	// Maybe will re-enable in the future if things change.
-	// TranslationData->SetFlags(RF_Public); //Otherwise IsAsset() fails
-
 	// We want Undo/Redo support
 	TranslationData->SetFlags(RF_Transactional);
 
@@ -394,12 +389,24 @@ void FTranslationDataManager::GetHistoryForTranslationUnits( TArray<FTranslation
 {
 	GWarn->BeginSlowTask(LOCTEXT("LoadingSourceControlHistory", "Loading Translation History from Source Control..."), true);
 
+	// Force history update
 	ISourceControlProvider& SourceControlProvider = ISourceControlModule::Get().GetProvider();
 	TSharedRef<FUpdateStatus, ESPMode::ThreadSafe> UpdateStatusOperation = ISourceControlOperation::Create<FUpdateStatus>();
 	UpdateStatusOperation->SetUpdateHistory( true );
-	SourceControlProvider.Execute( UpdateStatusOperation, ManifestFilePath );
-	FSourceControlStatePtr SourceControlState = SourceControlProvider.GetState( ManifestFilePath, EStateCacheUsage::ForceUpdate );
-	if ( SourceControlState.IsValid() )
+	ECommandResult::Type Result = SourceControlProvider.Execute(UpdateStatusOperation, ManifestFilePath);
+	bool bGetHistoryFromSourceControlSucceeded = Result == ECommandResult::Succeeded;
+
+	// Now we can get information about the file's history from the source control state, retrieve that
+	TArray<FString> Files;
+	TArray< TSharedRef<ISourceControlState, ESPMode::ThreadSafe> > States;
+	Files.Add(ManifestFilePath);
+	Result = SourceControlProvider.GetState( Files,  States, EStateCacheUsage::ForceUpdate );
+	bGetHistoryFromSourceControlSucceeded = bGetHistoryFromSourceControlSucceeded && (Result == ECommandResult::Succeeded);
+	check(States.Num() == 1);
+	FSourceControlStatePtr SourceControlState = States[0];
+	
+	// If all the source control operations went ok, continue
+	if (bGetHistoryFromSourceControlSucceeded && SourceControlState.IsValid())
 	{
 		int32 HistorySize = SourceControlState->GetHistorySize();
 
@@ -497,7 +504,8 @@ void FTranslationDataManager::GetHistoryForTranslationUnits( TArray<FTranslation
 			}
 		}
 	}
-	else // SourceControlState.IsValid() is false
+	// If source control operations failed, display error message
+	else // (bGetHistoryFromSourceControlSucceeded && SourceControlState.IsValid()) is false
 	{
 		FFormatNamedArguments Arguments;
 		Arguments.Add( TEXT("ManifestFilePath"), FText::FromString(ManifestFilePath) );
diff --git a/Engine/Source/Editor/TranslationEditor/Private/TranslationEditor.cpp b/Engine/Source/Editor/TranslationEditor/Private/TranslationEditor.cpp
--- a/Engine/Source/Editor/TranslationEditor/Private/TranslationEditor.cpp
+++ b/Engine/Source/Editor/TranslationEditor/Private/TranslationEditor.cpp
@@ -622,11 +622,26 @@ void FTranslationEditor::RefreshUI()
 	PreviewTextBlock->SetFont(TranslationTargetFont);
 
 	// Refresh our widget displays
-	UntranslatedPropertyTableWidgetHandle->RequestRefresh();
-	ReviewPropertyTableWidgetHandle->RequestRefresh();
-	CompletedPropertyTableWidgetHandle->RequestRefresh();
-	ContextPropertyTableWidgetHandle->RequestRefresh();
-	HistoryPropertyTableWidgetHandle->RequestRefresh();
+	if (UntranslatedPropertyTableWidgetHandle.IsValid())
+	{
+		UntranslatedPropertyTableWidgetHandle->RequestRefresh();
+	}
+	if (ReviewPropertyTableWidgetHandle.IsValid())	
+	{
+		ReviewPropertyTableWidgetHandle->RequestRefresh();
+	}
+	if (CompletedPropertyTableWidgetHandle.IsValid())
+	{
+		CompletedPropertyTableWidgetHandle->RequestRefresh();
+	}
+	if (ContextPropertyTableWidgetHandle.IsValid())
+	{
+		ContextPropertyTableWidgetHandle->RequestRefresh();
+	}
+	if (HistoryPropertyTableWidgetHandle.IsValid())
+	{
+		HistoryPropertyTableWidgetHandle->RequestRefresh();
+	}
 }
 
 bool FTranslationEditor::OpenFontPicker( const FString DefaultFile, FString& OutFile )
@@ -676,23 +691,26 @@ bool FTranslationEditor::OpenFontPicker( const FString DefaultFile, FString& Out
 void FTranslationEditor::UpdateTranslationUnitSelection()
 {
 	TSet<TSharedRef<IPropertyTableRow>> SelectedRows;
-	TSharedRef<FPropertyPath> InitialPath = ContextPropertyTable->GetRootPath(); // Dummy Path for now
+	TSharedRef<FPropertyPath> InitialPath = FPropertyPath::CreateEmpty(); // Dummy Path for now
 	UProperty* PropertyToFind = NULL;
 
 	// Different Selection based on which tab is in the foreground
-	if (UntranslatedTab.IsValid() && UntranslatedTab->IsForeground())
+	TSharedPtr<SDockTab> UntranslatedTabSharedPtr = UntranslatedTab.Pin();
+	TSharedPtr<SDockTab> ReviewTabSharedPtr = ReviewTab.Pin();
+	TSharedPtr<SDockTab> CompletedTabSharedPtr = CompletedTab.Pin();
+	if (UntranslatedTabSharedPtr.IsValid() && UntranslatedTabSharedPtr->IsForeground() && UntranslatedPropertyTable.IsValid())
 	{
 		SelectedRows = UntranslatedPropertyTable->GetSelectedRows();
 		InitialPath = UntranslatedPropertyTable->GetRootPath();
 		PropertyToFind = FindField<UProperty>( UTranslationDataObject::StaticClass(), "Untranslated");
 	}
-	else if (ReviewTab.IsValid() && ReviewTab->IsForeground())
+	else if (ReviewTab.IsValid() && ReviewTabSharedPtr->IsForeground() && ReviewPropertyTable.IsValid())
 	{
 		SelectedRows = ReviewPropertyTable->GetSelectedRows();
 		InitialPath = ReviewPropertyTable->GetRootPath();
 		PropertyToFind = FindField<UProperty>( UTranslationDataObject::StaticClass(), "Review");
 	}
-	else if (CompletedTab.IsValid() && CompletedTab->IsForeground())
+	else if (CompletedTab.IsValid() && CompletedTabSharedPtr->IsForeground() && CompletedPropertyTable.IsValid())
 	{
 		SelectedRows = CompletedPropertyTable->GetSelectedRows();
 		InitialPath = CompletedPropertyTable->GetRootPath();
@@ -718,20 +736,24 @@ void FTranslationEditor::UpdateTranslationUnitSelection()
 		TSharedRef<FPropertyPath> ContextPath = InitialPath;
 		ContextPath = ContextPath->ExtendPath(PartialPath);
 		ContextPath = ContextPath->ExtendPath(ContextArrayPropInfo);
-		ContextPropertyTable->SetRootPath(ContextPath);
+		
+		if (ContextPropertyTable.IsValid())
+		{
+			ContextPropertyTable->SetRootPath(ContextPath);
 
-		// Need to re-add the columns we want to display
-		ContextPropertyTable->AddColumn((TWeakObjectPtr<UProperty>)FindField<UProperty>( FTranslationContextInfo::StaticStruct(), "Key"));
-		ContextPropertyTable->AddColumn((TWeakObjectPtr<UProperty>)FindField<UProperty>( FTranslationContextInfo::StaticStruct(), "Context"));
+			// Need to re-add the columns we want to display
+			ContextPropertyTable->AddColumn((TWeakObjectPtr<UProperty>)FindField<UProperty>(FTranslationContextInfo::StaticStruct(), "Key"));
+			ContextPropertyTable->AddColumn((TWeakObjectPtr<UProperty>)FindField<UProperty>(FTranslationContextInfo::StaticStruct(), "Context"));
 
-		TArray<TSharedRef<IPropertyTableColumn>> Columns = ContextPropertyTable->GetColumns();
-		for (TSharedRef<IPropertyTableColumn> Column : Columns)
-		{
-			Column->SetFrozen(true);
+			TArray<TSharedRef<IPropertyTableColumn>> Columns = ContextPropertyTable->GetColumns();
+			for (TSharedRef<IPropertyTableColumn> Column : Columns)
+			{
+				Column->SetFrozen(true);
+			}
 		}
 
 		// Add the HistoryPropertyTable-specific path
-		TSharedRef<FPropertyPath> HistoryPath = ContextPropertyTable->GetRootPath();
+		TSharedRef<FPropertyPath> HistoryPath = ContextPath;
 		FPropertyInfo ContextPropInfo;
 		ContextPropInfo.Property = ContextArrayProp->Inner;
 		ContextPropInfo.ArrayIndex = 0;	// Just show history for the first context until the user selects something else
@@ -741,24 +763,24 @@ void FTranslationEditor::UpdateTranslationUnitSelection()
 		ChangesPropInfo.Property = ChangesProp;
 		ChangesPropInfo.ArrayIndex = INDEX_NONE;
 		HistoryPath = HistoryPath->ExtendPath(ChangesPropInfo);
-		HistoryPropertyTable->SetRootPath(HistoryPath);
-
-		// Need to re-add the columns we want to display
-		HistoryPropertyTable->AddColumn((TWeakObjectPtr<UProperty>)FindField<UProperty>( FTranslationChange::StaticStruct(), "Version"));
-		HistoryPropertyTable->AddColumn((TWeakObjectPtr<UProperty>)FindField<UProperty>( FTranslationChange::StaticStruct(), "DateAndTime"));
-		HistoryPropertyTable->AddColumn((TWeakObjectPtr<UProperty>)FindField<UProperty>( FTranslationChange::StaticStruct(), "Source"));
-		HistoryPropertyTable->AddColumn((TWeakObjectPtr<UProperty>)FindField<UProperty>( FTranslationChange::StaticStruct(), "Translation"));
 
-		Columns = HistoryPropertyTable->GetColumns();
-		for (TSharedRef<IPropertyTableColumn> Column : Columns)
+		if (HistoryPropertyTable.IsValid())
 		{
-			Column->SetFrozen(true);
+			HistoryPropertyTable->SetRootPath(HistoryPath);
+
+			// Need to re-add the columns we want to display
+			HistoryPropertyTable->AddColumn((TWeakObjectPtr<UProperty>)FindField<UProperty>(FTranslationChange::StaticStruct(), "Version"));
+			HistoryPropertyTable->AddColumn((TWeakObjectPtr<UProperty>)FindField<UProperty>(FTranslationChange::StaticStruct(), "DateAndTime"));
+			HistoryPropertyTable->AddColumn((TWeakObjectPtr<UProperty>)FindField<UProperty>(FTranslationChange::StaticStruct(), "Source"));
+			HistoryPropertyTable->AddColumn((TWeakObjectPtr<UProperty>)FindField<UProperty>(FTranslationChange::StaticStruct(), "Translation"));
+
+			TArray<TSharedRef<IPropertyTableColumn>> Columns = HistoryPropertyTable->GetColumns();
+			for (TSharedRef<IPropertyTableColumn> Column : Columns)
+			{
+				Column->SetFrozen(true);
+			}
 		}
-	
-		// Maybe we should save every time the user changes their selection to be safe?
-		//DataManager->WriteTranslationData();
 	}
-
 }
 
 void FTranslationEditor::SaveAsset_Execute()
@@ -769,66 +791,76 @@ void FTranslationEditor::SaveAsset_Execute()
 
 void FTranslationEditor::UpdateContextSelection()
 {
-	TSet<TSharedRef<IPropertyTableRow>> SelectedRows = ContextPropertyTable->GetSelectedRows();
-	TSharedRef<FPropertyPath> InitialPath = ContextPropertyTable->GetRootPath();
-	UProperty* PropertyToFind = InitialPath->GetRootProperty().Property.Get();
-	
-	// Can only really handle single selection
-	if (SelectedRows.Num() == 1)
+	if (ContextPropertyTable.IsValid())
 	{
-		TSharedRef<IPropertyTableRow> SelectedRow = *(SelectedRows.CreateConstIterator());
-		TSharedRef<FPropertyPath> PartialPath = SelectedRow->GetPartialPath();
-		
-		TArray<FTranslationUnit>& TranslationUnitArray = *(PropertyToFind->ContainerPtrToValuePtr<TArray<FTranslationUnit>>(TranslationData));
-
-		// Index of the next to last property in the list is the index in the array of translation units
-		FTranslationUnit& SelectedTranslationUnit = TranslationUnitArray[InitialPath->GetPropertyInfo(InitialPath->GetNumProperties()-2).ArrayIndex];
-		// Index of the leaf most property is the context info index we need
-		FTranslationContextInfo& SelectedContextInfo = SelectedTranslationUnit.Contexts[PartialPath->GetLeafMostProperty().ArrayIndex];
+		TSet<TSharedRef<IPropertyTableRow>> SelectedRows = ContextPropertyTable->GetSelectedRows();
+		TSharedRef<FPropertyPath> InitialPath = ContextPropertyTable->GetRootPath();
+		UProperty* PropertyToFind = InitialPath->GetRootProperty().Property.Get();
 
-		// If this is a translation unit from the review tab and they select a context, possibly update the selected translation with one from that context
-		if (InitialPath->GetRootProperty().Property == FindField<UProperty>( UTranslationDataObject::StaticClass(), "Review"))
+		// Can only really handle single selection
+		if (SelectedRows.Num() == 1)
 		{
-			// Only change the suggested translation if they haven't yet reviewed it
-			if (SelectedTranslationUnit.HasBeenReviewed == false)
+			TSharedRef<IPropertyTableRow> SelectedRow = *(SelectedRows.CreateConstIterator());
+			TSharedRef<FPropertyPath> PartialPath = SelectedRow->GetPartialPath();
+
+			TArray<FTranslationUnit>& TranslationUnitArray = *(PropertyToFind->ContainerPtrToValuePtr<TArray<FTranslationUnit>>(TranslationData));
+
+			// Index of the next to last property in the list is the index in the array of translation units
+			FTranslationUnit& SelectedTranslationUnit = TranslationUnitArray[InitialPath->GetPropertyInfo(InitialPath->GetNumProperties() - 2).ArrayIndex];
+			// Index of the leaf most property is the context info index we need
+			FTranslationContextInfo& SelectedContextInfo = SelectedTranslationUnit.Contexts[PartialPath->GetLeafMostProperty().ArrayIndex];
+
+			// If this is a translation unit from the review tab and they select a context, possibly update the selected translation with one from that context
+			if (InitialPath->GetRootProperty().Property == FindField<UProperty>(UTranslationDataObject::StaticClass(), "Review"))
 			{
-				for (int32 ChangeIndex = 0; ChangeIndex < SelectedContextInfo.Changes.Num(); ++ChangeIndex)
+				// Only change the suggested translation if they haven't yet reviewed it
+				if (SelectedTranslationUnit.HasBeenReviewed == false)
 				{
-					// Find most recent, non-empty translation
-					if (!SelectedContextInfo.Changes[ChangeIndex].Translation.IsEmpty())
+					for (int32 ChangeIndex = 0; ChangeIndex < SelectedContextInfo.Changes.Num(); ++ChangeIndex)
 					{
-						TranslationData->Modify();
-						SelectedTranslationUnit.Translation = SelectedContextInfo.Changes[ChangeIndex].Translation;
-						TranslationData->PostEditChange();
+						// Find most recent, non-empty translation
+						if (!SelectedContextInfo.Changes[ChangeIndex].Translation.IsEmpty())
+						{
+							check(TranslationData != NULL);
+							if (TranslationData != NULL)
+							{
+								TranslationData->Modify();
+								SelectedTranslationUnit.Translation = SelectedContextInfo.Changes[ChangeIndex].Translation;
+								TranslationData->PostEditChange();
+							}
+						}
 					}
 				}
 			}
-		}
 
-		// Add the HistoryPropertyTable-specific path
-		TSharedRef<FPropertyPath> HistoryPath = ContextPropertyTable->GetRootPath();
-		UArrayProperty* ContextArrayProp = FindField<UArrayProperty>( FTranslationUnit::StaticStruct(), "Contexts" );
-		FPropertyInfo ContextPropInfo;
-		ContextPropInfo.Property = ContextArrayProp->Inner;
-		ContextPropInfo.ArrayIndex = PartialPath->GetLeafMostProperty().ArrayIndex;
-		HistoryPath = HistoryPath->ExtendPath(ContextPropInfo);
-		UArrayProperty* ChangesProp = FindField<UArrayProperty>( FTranslationContextInfo::StaticStruct(), "Changes" );
-		FPropertyInfo ChangesPropInfo;
-		ChangesPropInfo.Property = ChangesProp;
-		ChangesPropInfo.ArrayIndex = INDEX_NONE;
-		HistoryPath = HistoryPath->ExtendPath(ChangesPropInfo);
-		HistoryPropertyTable->SetRootPath(HistoryPath);
+			// Add the HistoryPropertyTable-specific path
+			TSharedRef<FPropertyPath> HistoryPath = ContextPropertyTable->GetRootPath();
+			UArrayProperty* ContextArrayProp = FindField<UArrayProperty>(FTranslationUnit::StaticStruct(), "Contexts");
+			FPropertyInfo ContextPropInfo;
+			ContextPropInfo.Property = ContextArrayProp->Inner;
+			ContextPropInfo.ArrayIndex = PartialPath->GetLeafMostProperty().ArrayIndex;
+			HistoryPath = HistoryPath->ExtendPath(ContextPropInfo);
+			UArrayProperty* ChangesProp = FindField<UArrayProperty>(FTranslationContextInfo::StaticStruct(), "Changes");
+			FPropertyInfo ChangesPropInfo;
+			ChangesPropInfo.Property = ChangesProp;
+			ChangesPropInfo.ArrayIndex = INDEX_NONE;
+			HistoryPath = HistoryPath->ExtendPath(ChangesPropInfo);
+			if (HistoryPropertyTable.IsValid())
+			{
+				HistoryPropertyTable->SetRootPath(HistoryPath);
 
-		// Need to re-add the columns we want to display
-		HistoryPropertyTable->AddColumn((TWeakObjectPtr<UProperty>)FindField<UProperty>( FTranslationChange::StaticStruct(), "Version"));
-		HistoryPropertyTable->AddColumn((TWeakObjectPtr<UProperty>)FindField<UProperty>( FTranslationChange::StaticStruct(), "DateAndTime"));
-		HistoryPropertyTable->AddColumn((TWeakObjectPtr<UProperty>)FindField<UProperty>( FTranslationChange::StaticStruct(), "Source"));
-		HistoryPropertyTable->AddColumn((TWeakObjectPtr<UProperty>)FindField<UProperty>( FTranslationChange::StaticStruct(), "Translation"));
+				// Need to re-add the columns we want to display
+				HistoryPropertyTable->AddColumn((TWeakObjectPtr<UProperty>)FindField<UProperty>(FTranslationChange::StaticStruct(), "Version"));
+				HistoryPropertyTable->AddColumn((TWeakObjectPtr<UProperty>)FindField<UProperty>(FTranslationChange::StaticStruct(), "DateAndTime"));
+				HistoryPropertyTable->AddColumn((TWeakObjectPtr<UProperty>)FindField<UProperty>(FTranslationChange::StaticStruct(), "Source"));
+				HistoryPropertyTable->AddColumn((TWeakObjectPtr<UProperty>)FindField<UProperty>(FTranslationChange::StaticStruct(), "Translation"));
 
-		TArray<TSharedRef<IPropertyTableColumn>> Columns = HistoryPropertyTable->GetColumns();
-		for (TSharedRef<IPropertyTableColumn> Column : Columns)
-		{
-			Column->SetFrozen(true);
+				TArray<TSharedRef<IPropertyTableColumn>> Columns = HistoryPropertyTable->GetColumns();
+				for (TSharedRef<IPropertyTableColumn> Column : Columns)
+				{
+					Column->SetFrozen(true);
+				}
+			}
 		}
 	}
 }
diff --git a/Engine/Source/Editor/TranslationEditor/Private/TranslationEditor.h b/Engine/Source/Editor/TranslationEditor/Private/TranslationEditor.h
--- a/Engine/Source/Editor/TranslationEditor/Private/TranslationEditor.h
+++ b/Engine/Source/Editor/TranslationEditor/Private/TranslationEditor.h
@@ -150,11 +150,11 @@ private:
 	static const FName HistoryTabId;
 
 	/** The Untranslated Tab */
-	TSharedPtr<SDockTab> UntranslatedTab;
+	TWeakPtr<SDockTab> UntranslatedTab;
 	/** The Review Tab */
-	TSharedPtr<SDockTab> ReviewTab;
+	TWeakPtr<SDockTab> ReviewTab;
 	/** The Review Tab */
-	TSharedPtr<SDockTab> CompletedTab;
+	TWeakPtr<SDockTab> CompletedTab;
 
 	/** UObject containing our translation information */
 	UTranslationDataObject* TranslationData;
