Translation Editor:
- Removed check() that resulted in assert if not connected to source control.

diff --git a/Engine/Source/Editor/TranslationEditor/Private/TranslationDataManager.cpp b/Engine/Source/Editor/TranslationEditor/Private/TranslationDataManager.cpp
--- a/Engine/Source/Editor/TranslationEditor/Private/TranslationDataManager.cpp
+++ b/Engine/Source/Editor/TranslationEditor/Private/TranslationDataManager.cpp
@@ -402,9 +402,12 @@ void FTranslationDataManager::GetHistoryForTranslationUnits( TArray<FTranslation
 	Files.Add(ManifestFilePath);
 	Result = SourceControlProvider.GetState( Files,  States, EStateCacheUsage::ForceUpdate );
 	bGetHistoryFromSourceControlSucceeded = bGetHistoryFromSourceControlSucceeded && (Result == ECommandResult::Succeeded);
-	check(States.Num() == 1);
-	FSourceControlStatePtr SourceControlState = States[0];
-	
+	FSourceControlStatePtr SourceControlState;
+	if (States.Num() == 1)
+	{
+		SourceControlState = States[0];
+	}
+
 	// If all the source control operations went ok, continue
 	if (bGetHistoryFromSourceControlSucceeded && SourceControlState.IsValid())
 	{
