Translation Editor UI
- New toolbar button to Preview all translation strings in the editor

diff --git a/Engine/Source/Editor/TranslationEditor/Private/TranslationDataManager.cpp b/Engine/Source/Editor/TranslationEditor/Private/TranslationDataManager.cpp
--- a/Engine/Source/Editor/TranslationEditor/Private/TranslationDataManager.cpp
+++ b/Engine/Source/Editor/TranslationEditor/Private/TranslationDataManager.cpp
@@ -8,18 +8,16 @@
 #include "InternationalizationArchiveJsonSerializer.h"
 #include "ISourceControlModule.h"
 #include "MessageLog.h"
+#include "TextLocalizationManager.h"
 
 DEFINE_LOG_CATEGORY_STATIC(LogTranslationEditor, Log, All);
 
 #define LOCTEXT_NAMESPACE "TranslationDataManager"
 
-FTranslationDataManager::FTranslationDataManager( const FString& ManifestFile, const FString& ArchiveFile )
+FTranslationDataManager::FTranslationDataManager( const FString& InManifestFilePath, const FString& InArchiveFilePath )
+: ManifestFilePath(InManifestFilePath)
+, ArchiveFilePath(InArchiveFilePath)
 {
-	FString ManifestName = FPaths::GetCleanFilename(ManifestFile);
-	FString ProjectPath = FPaths::GetPath(ManifestFile);
-	FString ArchiveName = FPaths::GetCleanFilename(ArchiveFile);
-	FString CulturePath = FPaths::GetPath(ArchiveFile);
-
 	GWarn->BeginSlowTask(LOCTEXT("LoadingTranslationData", "Loading Translation Data..."), true);
 
 	TranslationData = NewObject<UTranslationDataObject>();
@@ -28,7 +26,6 @@ FTranslationDataManager::FTranslationDataManager( const FString& ManifestFile, c
 	// We want Undo/Redo support
 	TranslationData->SetFlags(RF_Transactional);
 
-	FString ManifestFilePath = ProjectPath / ManifestName;
 	ManifestAtHeadRevisionPtr = ReadManifest( ManifestFilePath );
 	if (ManifestAtHeadRevisionPtr.IsValid())
 	{
@@ -46,7 +43,6 @@ FTranslationDataManager::FTranslationDataManager( const FString& ManifestFile, c
 			TranslationEditorMessageLog.Open(EMessageSeverity::Error);
 		}
 
-		ArchiveFilePath = CulturePath + TEXT("/") + ArchiveName;
 		ArchivePtr = ReadArchive( ManifestAtHeadRevision);
 		if (ArchivePtr.IsValid())
 		{
@@ -166,14 +162,14 @@ FTranslationDataManager::FTranslationDataManager( const FString& ManifestFile, c
 	GWarn->EndSlowTask();
 }
 
-TSharedPtr< FInternationalizationManifest > FTranslationDataManager::ReadManifest( const FString& ManifestFilePath )
+TSharedPtr< FInternationalizationManifest > FTranslationDataManager::ReadManifest( const FString& ManifestFilePathToRead )
 {
 	
-	TSharedPtr<FJsonObject> ManifestJsonObject = ReadJSONTextFile( ManifestFilePath );
+	TSharedPtr<FJsonObject> ManifestJsonObject = ReadJSONTextFile( ManifestFilePathToRead );
 
 	if( !ManifestJsonObject.IsValid() )
 	{
-		UE_LOG(LogTranslationEditor, Error, TEXT("Could not read manifest file %s."), *ManifestFilePath);
+		UE_LOG(LogTranslationEditor, Error, TEXT("Could not read manifest file %s."), *ManifestFilePathToRead);
 		return TSharedPtr< FInternationalizationManifest >();
 	}
 
@@ -526,4 +522,30 @@ void FTranslationDataManager::HandlePropertyChanged(FName PropertyName)
 	WriteTranslationData();
 }
 
+void FTranslationDataManager::PreviewAllTranslationsInEditor()
+{
+	FString ManifestFullPath = FPaths::ConvertRelativePathToFull(ManifestFilePath);
+	FString EngineFullPath = FPaths::ConvertRelativePathToFull(FPaths::EngineContentDir());
+
+	bool IsEngineManifest = false;
+	if (ManifestFullPath.StartsWith(EngineFullPath))
+	{
+		IsEngineManifest = true;
+	}
+
+	FString ConfigDirectory;
+	if (IsEngineManifest)
+	{
+		ConfigDirectory = FPaths::EngineConfigDir();
+	}
+	else
+	{
+		ConfigDirectory = FPaths::GameConfigDir();
+	}
+
+	FString ConfigFilePath = ConfigDirectory / "Localization" / "Regenerate" + FPaths::GetBaseFilename(ManifestFilePath) + ".ini";
+
+	FTextLocalizationManager::Get().RegenerateResources(ConfigFilePath);
+}
+
 #undef LOCTEXT_NAMESPACE
diff --git a/Engine/Source/Editor/TranslationEditor/Private/TranslationDataManager.h b/Engine/Source/Editor/TranslationEditor/Private/TranslationDataManager.h
--- a/Engine/Source/Editor/TranslationEditor/Private/TranslationDataManager.h
+++ b/Engine/Source/Editor/TranslationEditor/Private/TranslationDataManager.h
@@ -15,7 +15,7 @@ class FTranslationDataManager : public TSharedFromThis<FTranslationDataManager>
 {
 
 public:
-	FTranslationDataManager( const FString& ManifestFile, const FString& ArchiveFile);
+	FTranslationDataManager( const FString& InManifestFilePath, const FString& InArchiveFilePath);
 
 	UTranslationDataObject* GetTranslationDataObject() 
 	{
@@ -28,6 +28,9 @@ public:
 	/** Delegate called when a TranslationDataObject property is changed */
 	void HandlePropertyChanged(FName PropertyName);
 
+	/** Regenerate and reload archives to reflect modifications in the UI */
+	void PreviewAllTranslationsInEditor();
+
 private:
 	/** Read text file into a JSON file */
 	static TSharedPtr<FJsonObject> ReadJSONTextFile( const FString& InFilePath ) ;
@@ -64,6 +67,8 @@ private:
 	FString ArchiveName;
 	/** Path to the culture (language, sort of) we are targeting */
 	FString CulturePath;
+	/** Path to the Manifest File **/
+	FString ManifestFilePath;
 	/** Path to the Archive File **/
 	FString ArchiveFilePath;
 
diff --git a/Engine/Source/Editor/TranslationEditor/Private/TranslationEditor.cpp b/Engine/Source/Editor/TranslationEditor/Private/TranslationEditor.cpp
--- a/Engine/Source/Editor/TranslationEditor/Private/TranslationEditor.cpp
+++ b/Engine/Source/Editor/TranslationEditor/Private/TranslationEditor.cpp
@@ -582,6 +582,10 @@ void FTranslationEditor::MapActions()
 	ToolkitCommands->MapAction( FTranslationEditorCommands::Get().SaveTranslations,
 		FExecuteAction::CreateSP(this, &FTranslationEditor::SaveAsset_Execute),
 		FCanExecuteAction());
+
+	ToolkitCommands->MapAction(FTranslationEditorCommands::Get().PreviewAllTranslationsInEditor,
+		FExecuteAction::CreateSP(this, &FTranslationEditor::PreviewAllTranslationsInEditor_Execute),
+		FCanExecuteAction());
 }
 
 void FTranslationEditor::ChangeSourceFont()
@@ -865,4 +869,9 @@ void FTranslationEditor::UpdateContextSelection()
 	}
 }
 
+void FTranslationEditor::PreviewAllTranslationsInEditor_Execute()
+{
+	DataManager->PreviewAllTranslationsInEditor();
+}
+
 #undef LOCTEXT_NAMESPACE
diff --git a/Engine/Source/Editor/TranslationEditor/Private/TranslationEditor.h b/Engine/Source/Editor/TranslationEditor/Private/TranslationEditor.h
--- a/Engine/Source/Editor/TranslationEditor/Private/TranslationEditor.h
+++ b/Engine/Source/Editor/TranslationEditor/Private/TranslationEditor.h
@@ -136,6 +136,9 @@ private:
 	/** Update content when a new context selection is made */
 	void UpdateContextSelection();
 
+	/** Called when "Apply to UI" is clicked for this asset */
+	void PreviewAllTranslationsInEditor_Execute();
+
 	/**	The tab id for the untranslated tab */
 	static const FName UntranslatedTabId;
 	/**	The tab id for the review tab */
diff --git a/Engine/Source/Editor/TranslationEditor/Private/TranslationEditorMenu.cpp b/Engine/Source/Editor/TranslationEditor/Private/TranslationEditorMenu.cpp
--- a/Engine/Source/Editor/TranslationEditor/Private/TranslationEditorMenu.cpp
+++ b/Engine/Source/Editor/TranslationEditor/Private/TranslationEditorMenu.cpp
@@ -15,6 +15,7 @@ void FTranslationEditorMenu::FillTranslationMenu( FMenuBuilder& MenuBuilder/*, F
 	{
 		MenuBuilder.AddMenuEntry( FTranslationEditorCommands::Get().ChangeSourceFont );
 		MenuBuilder.AddMenuEntry( FTranslationEditorCommands::Get().ChangeTranslationTargetFont );
+		MenuBuilder.AddMenuEntry( FTranslationEditorCommands::Get().PreviewAllTranslationsInEditor );
 	}
 	MenuBuilder.EndSection();
 }
@@ -63,6 +64,8 @@ void FTranslationEditorMenu::SetupTranslationEditorToolbar( TSharedPtr< FExtende
 		{
 			ToolbarBuilder.AddToolBarButton(
 				FTranslationEditorCommands::Get().SaveTranslations, "SaveTranslations", TAttribute<FText>(), TAttribute<FText>(), FSlateIcon(FEditorStyle::GetStyleSetName(), "AssetEditor.SaveAsset"));
+			ToolbarBuilder.AddToolBarButton(
+				FTranslationEditorCommands::Get().PreviewAllTranslationsInEditor, "ApplyTranslationsToUI", TAttribute<FText>(), TAttribute<FText>(), FSlateIcon(FEditorStyle::GetStyleSetName(), "AssetEditor.ReimportAsset"));
 		}
 	};
 
@@ -82,6 +85,7 @@ void FTranslationEditorCommands::RegisterCommands()
 	UI_COMMAND( ChangeSourceFont, "Change Source Font", "Change the Font for the Source Lanugage", EUserInterfaceActionType::Button, FInputGesture() );
 	UI_COMMAND( ChangeTranslationTargetFont, "Change Translation Font", "Change the Translation Target Language Font", EUserInterfaceActionType::Button, FInputGesture() );
 	UI_COMMAND( SaveTranslations, "Save", "Saves the translations to file", EUserInterfaceActionType::Button, FInputGesture() );
+	UI_COMMAND(PreviewAllTranslationsInEditor, "Preview in Editor", "Preview All Translations in the Editor UI", EUserInterfaceActionType::Button, FInputGesture());
 }
 
 
diff --git a/Engine/Source/Editor/TranslationEditor/Public/TranslationEditorMenu.h b/Engine/Source/Editor/TranslationEditor/Public/TranslationEditorMenu.h
--- a/Engine/Source/Editor/TranslationEditor/Public/TranslationEditorMenu.h
+++ b/Engine/Source/Editor/TranslationEditor/Public/TranslationEditorMenu.h
@@ -33,6 +33,9 @@ public:
 	/** Save the translations to file */
 	TSharedPtr<FUICommandInfo> SaveTranslations;
 
+	/** Save the translations to file */
+	TSharedPtr<FUICommandInfo> PreviewAllTranslationsInEditor;
+
 	/** Initialize commands */
 	virtual void RegisterCommands() OVERRIDE;
 };
diff --git a/Engine/Source/Runtime/Core/Private/Misc/Paths.cpp b/Engine/Source/Runtime/Core/Private/Misc/Paths.cpp
--- a/Engine/Source/Runtime/Core/Private/Misc/Paths.cpp
+++ b/Engine/Source/Runtime/Core/Private/Misc/Paths.cpp
@@ -87,6 +87,11 @@ FString FPaths::GameContentDir()
 	return FPaths::GameDir() + TEXT("Content/");
 }
 
+FString FPaths::GameConfigDir()
+{
+	return FPaths::GameDir() + TEXT("Config/");
+}
+
 FString FPaths::GameSavedDir()
 {
 	return GameUserDir() + TEXT("Saved/");
@@ -839,4 +844,5 @@ void FPaths::CombineInternal(FString& OutPath, const TCHAR** Pathes, int32 NumPa
 	{
 		OutPath /= Pathes[i];
 	}
-}
\ No newline at end of file
+}
+
diff --git a/Engine/Source/Runtime/Core/Public/Misc/Paths.h b/Engine/Source/Runtime/Core/Public/Misc/Paths.h
--- a/Engine/Source/Runtime/Core/Public/Misc/Paths.h
+++ b/Engine/Source/Runtime/Core/Public/Misc/Paths.h
@@ -96,6 +96,13 @@ public:
 	 */
 	static FString GameContentDir();
 
+	/**
+	* Returns the directory the root configuration files are located.
+	*
+	* @return root config directory
+	*/
+	static FString GameConfigDir();
+
 	/**
 	 * Returns the saved directory of the current game by looking at the global
 	 * GGameName variable. 
