Translation Editor:

- Export to PortableObject (.po) format button is now functional for Editor & Engine strings
- For now just uses settings in Engine\\Config\\Localization\\PortableObjectExport.ini, and exports both projects in ja, ko, zh

diff --git a/Engine/Source/Editor/TranslationEditor/Private/TranslationEditor.cpp b/Engine/Source/Editor/TranslationEditor/Private/TranslationEditor.cpp
--- a/Engine/Source/Editor/TranslationEditor/Private/TranslationEditor.cpp
+++ b/Engine/Source/Editor/TranslationEditor/Private/TranslationEditor.cpp
@@ -21,6 +21,8 @@
 
 #include "IPropertyTableWidgetHandle.h"
 
+DEFINE_LOG_CATEGORY_STATIC(LocalizationExport, Log, All);
+
 #define LOCTEXT_NAMESPACE "TranslationEditor"
 
 const FName FTranslationEditor::UntranslatedTabId( TEXT( "TranslationEditor_Untranslated" ) );
@@ -586,6 +588,10 @@ void FTranslationEditor::MapActions()
 	ToolkitCommands->MapAction(FTranslationEditorCommands::Get().PreviewAllTranslationsInEditor,
 		FExecuteAction::CreateSP(this, &FTranslationEditor::PreviewAllTranslationsInEditor_Execute),
 		FCanExecuteAction());
+
+	ToolkitCommands->MapAction(FTranslationEditorCommands::Get().ExportToPortableObjectFormat,
+		FExecuteAction::CreateSP(this, &FTranslationEditor::ExportToPortableObjectFormat_Execute),
+		FCanExecuteAction());
 }
 
 void FTranslationEditor::ChangeSourceFont()
@@ -651,7 +657,7 @@ void FTranslationEditor::RefreshUI()
 bool FTranslationEditor::OpenFontPicker( const FString DefaultFile, FString& OutFile )
 {
 	const FString FontFileDescription = LOCTEXT( "FontFileDescription", "Font File" ).ToString();
-	//TODO: support more than one filetype (right now only .ttfs are supported)
+	//TODO: support more than one filetype (ttf and otf are supported)
 	const FString FontFileExtension = TEXT("*.ttf;*.otf");
 	const FString FileTypes = FString::Printf( TEXT("%s (%s)|%s"), *FontFileDescription, *FontFileExtension, *FontFileExtension );
 
@@ -874,4 +880,94 @@ void FTranslationEditor::PreviewAllTranslationsInEditor_Execute()
 	DataManager->PreviewAllTranslationsInEditor();
 }
 
+void FTranslationEditor::ExportToPortableObjectFormat_Execute()
+{
+	// TODO: Add ability to choose output location
+
+	//const FString PortableObjectFileDescription = LOCTEXT("PortableObjectFileDescription", "Portable Object File").ToString();
+	//const FString PortableObjectFileExtension = TEXT("*.po");
+	//const FString FileTypes = FString::Printf(TEXT("%s (%s)|%s"), *PortableObjectFileDescription, *PortableObjectFileExtension, *PortableObjectFileExtension);
+	TArray<FString> OpenFilenames;
+	//IDesktopPlatform* DesktopPlatform = FDesktopPlatformModule::Get();
+	bool bSelected = false;
+	
+	// Prompt the user for the filenames
+	//if (DesktopPlatform)
+	//{
+	//	void* ParentWindowWindowHandle = NULL;
+
+	//	const TSharedPtr<SWindow>& ParentWindow = FSlateApplication::Get().FindWidgetWindow(PreviewTextBlock);
+	//	if (ParentWindow.IsValid() && ParentWindow->GetNativeWindow().IsValid())
+	//	{
+	//		ParentWindowWindowHandle = ParentWindow->GetNativeWindow()->GetOSWindowHandle();
+	//	}
+
+	//	bSelected = DesktopPlatform->SaveFileDialog(
+	//		ParentWindowWindowHandle,
+	//		LOCTEXT("ChooseExportLocationWindowTitle", "Choose Export Location").ToString(),
+	//		FPaths::GameSavedDir(),
+	//		"ExportedLocalization.po",
+	//		FileTypes,
+	//		EFileDialogFlags::None,
+	//		OpenFilenames
+	//		);
+	//}
+
+	FString OutFile = FPaths::EngineSavedDir() / "LocalizationExport" / "ExportedLocalization.po";
+
+	if (bSelected && OpenFilenames.Num() > 0)
+	{
+		OutFile = OpenFilenames[0];
+	}
+
+	// For now, just run every section in the config file
+	TArray<FString> ConfigSections;
+	ConfigSections.Add("GatherTextStep0");
+	ConfigSections.Add("GatherTextStep1");
+	ConfigSections.Add("GatherTextStep2");
+	ConfigSections.Add("GatherTextStep3");
+
+	for (FString& ConfigSection : ConfigSections)
+	{
+		// Spawn the LocalizationExport commandlet, and run its log output back into ours
+		FString AppURL = FPlatformProcess::ExecutableName(true);
+		FString Parameters = FString("-run=InternationalizationExport -config=") + FPaths::EngineConfigDir() / "Localization" / "PortableObjectExport.ini -section=" + ConfigSection;
+
+		void* WritePipe;
+		void* ReadPipe;
+		FPlatformProcess::CreatePipe(ReadPipe, WritePipe);
+		FProcHandle ProcessHandle = FPlatformProcess::CreateProc(*AppURL, *Parameters, false, false, false, NULL, 0, NULL, WritePipe);
+
+		while (FPlatformProcess::IsProcRunning(ProcessHandle))
+		{
+			FString NewLine = FPlatformProcess::ReadPipe(ReadPipe);
+			if (NewLine.Len() > 0)
+			{
+				UE_LOG(LocalizationExport, Log, TEXT("%s"), *NewLine);
+			}
+
+			FPlatformProcess::Sleep(0.25);
+		}
+		FString NewLine = FPlatformProcess::ReadPipe(ReadPipe);
+		if (NewLine.Len() > 0)
+		{
+			UE_LOG(LocalizationExport, Log, TEXT("%s"), *NewLine);
+		}
+
+		FPlatformProcess::Sleep(0.25);
+		FPlatformProcess::ClosePipe(ReadPipe, WritePipe);
+
+		int32 ReturnCode;
+		if (!FPlatformProcess::GetProcReturnCode(ProcessHandle, &ReturnCode))
+		{
+			// TODO: print some error
+		}
+		else if (ReturnCode != 0)
+		{
+			// TODO: print some error
+		}
+	}
+
+}
+
 #undef LOCTEXT_NAMESPACE
diff --git a/Engine/Source/Editor/TranslationEditor/Private/TranslationEditor.h b/Engine/Source/Editor/TranslationEditor/Private/TranslationEditor.h
--- a/Engine/Source/Editor/TranslationEditor/Private/TranslationEditor.h
+++ b/Engine/Source/Editor/TranslationEditor/Private/TranslationEditor.h
@@ -136,9 +136,12 @@ private:
 	/** Update content when a new context selection is made */
 	void UpdateContextSelection();
 
-	/** Called when "Apply to UI" is clicked for this asset */
+	/** Called when "Preview in Editor" is clicked for this asset */
 	void PreviewAllTranslationsInEditor_Execute();
 
+	/** Called when "Export to .PO" is clicked for this asset */
+	void ExportToPortableObjectFormat_Execute();
+
 	/**	The tab id for the untranslated tab */
 	static const FName UntranslatedTabId;
 	/**	The tab id for the review tab */
diff --git a/Engine/Source/Editor/TranslationEditor/Private/TranslationEditorMenu.cpp b/Engine/Source/Editor/TranslationEditor/Private/TranslationEditorMenu.cpp
--- a/Engine/Source/Editor/TranslationEditor/Private/TranslationEditorMenu.cpp
+++ b/Engine/Source/Editor/TranslationEditor/Private/TranslationEditorMenu.cpp
@@ -16,6 +16,7 @@ void FTranslationEditorMenu::FillTranslationMenu( FMenuBuilder& MenuBuilder/*, F
 		MenuBuilder.AddMenuEntry( FTranslationEditorCommands::Get().ChangeSourceFont );
 		MenuBuilder.AddMenuEntry( FTranslationEditorCommands::Get().ChangeTranslationTargetFont );
 		MenuBuilder.AddMenuEntry( FTranslationEditorCommands::Get().PreviewAllTranslationsInEditor );
+		MenuBuilder.AddMenuEntry( FTranslationEditorCommands::Get().ExportToPortableObjectFormat );
 	}
 	MenuBuilder.EndSection();
 }
@@ -65,7 +66,9 @@ void FTranslationEditorMenu::SetupTranslationEditorToolbar( TSharedPtr< FExtende
 			ToolbarBuilder.AddToolBarButton(
 				FTranslationEditorCommands::Get().SaveTranslations, "SaveTranslations", TAttribute<FText>(), TAttribute<FText>(), FSlateIcon(FEditorStyle::GetStyleSetName(), "AssetEditor.SaveAsset"));
 			ToolbarBuilder.AddToolBarButton(
-				FTranslationEditorCommands::Get().PreviewAllTranslationsInEditor, "ApplyTranslationsToUI", TAttribute<FText>(), TAttribute<FText>(), FSlateIcon(FEditorStyle::GetStyleSetName(), "AssetEditor.ReimportAsset"));
+				FTranslationEditorCommands::Get().PreviewAllTranslationsInEditor, "PreviewTranslationsInEditor", TAttribute<FText>(), TAttribute<FText>(), FSlateIcon(FEditorStyle::GetStyleSetName(), "AssetEditor.ReimportAsset")); 
+			ToolbarBuilder.AddToolBarButton(
+				FTranslationEditorCommands::Get().ExportToPortableObjectFormat, "ExportToPortableObjectFormat", TAttribute<FText>(), TAttribute<FText>(), FSlateIcon(FEditorStyle::GetStyleSetName(), "FontEditor.Export"));
 		}
 	};
 
@@ -85,7 +88,8 @@ void FTranslationEditorCommands::RegisterCommands()
 	UI_COMMAND( ChangeSourceFont, "Change Source Font", "Change the Font for the Source Lanugage", EUserInterfaceActionType::Button, FInputGesture() );
 	UI_COMMAND( ChangeTranslationTargetFont, "Change Translation Font", "Change the Translation Target Language Font", EUserInterfaceActionType::Button, FInputGesture() );
 	UI_COMMAND( SaveTranslations, "Save", "Saves the translations to file", EUserInterfaceActionType::Button, FInputGesture() );
-	UI_COMMAND(PreviewAllTranslationsInEditor, "Preview in Editor", "Preview All Translations in the Editor UI", EUserInterfaceActionType::Button, FInputGesture());
+	UI_COMMAND( PreviewAllTranslationsInEditor, "Preview in Editor", "Preview All Translations in the Editor UI", EUserInterfaceActionType::Button, FInputGesture() );
+	UI_COMMAND( ExportToPortableObjectFormat, "Export to .PO", "Export to Portable Object Format", EUserInterfaceActionType::Button, FInputGesture() );
 }
 
 
diff --git a/Engine/Source/Editor/TranslationEditor/Public/TranslationEditorMenu.h b/Engine/Source/Editor/TranslationEditor/Public/TranslationEditorMenu.h
--- a/Engine/Source/Editor/TranslationEditor/Public/TranslationEditorMenu.h
+++ b/Engine/Source/Editor/TranslationEditor/Public/TranslationEditorMenu.h
@@ -36,6 +36,9 @@ public:
 	/** Save the translations to file */
 	TSharedPtr<FUICommandInfo> PreviewAllTranslationsInEditor;
 
+	/** Export to PortableObject format (.po) */
+	TSharedPtr<FUICommandInfo> ExportToPortableObjectFormat;
+
 	/** Initialize commands */
 	virtual void RegisterCommands() OVERRIDE;
 };
