TranslationEditor:
- Removed TranslationDataObject and converted TranslationUnit to a UObject instead of UStruct.
- Needed to do this to show the same TranslationUnit in multiple PropertyTables

diff --git a/Engine/Source/Editor/TranslationEditor/Private/TranslationDataManager.cpp b/Engine/Source/Editor/TranslationEditor/Private/TranslationDataManager.cpp
--- a/Engine/Source/Editor/TranslationEditor/Private/TranslationDataManager.cpp
+++ b/Engine/Source/Editor/TranslationEditor/Private/TranslationDataManager.cpp
@@ -3,7 +3,7 @@
 #include "TranslationEditorPrivatePCH.h"
 #include "TranslationEditor.h"
 #include "WorkspaceMenuStructureModule.h"
-#include "TranslationDataObject.h"
+#include "TranslationUnit.h"
 #include "InternationalizationManifestJsonSerializer.h"
 #include "InternationalizationArchiveJsonSerializer.h"
 #include "ISourceControlModule.h"
@@ -20,12 +20,6 @@ FTranslationDataManager::FTranslationDataManager( const FString& InManifestFileP
 {
 	GWarn->BeginSlowTask(LOCTEXT("LoadingTranslationData", "Loading Translation Data..."), true);
 
-	TranslationData = NewObject<UTranslationDataObject>();
-	check (TranslationData != NULL);
-
-	// We want Undo/Redo support
-	TranslationData->SetFlags(RF_Transactional);
-
 	ManifestAtHeadRevisionPtr = ReadManifest( ManifestFilePath );
 	if (ManifestAtHeadRevisionPtr.IsValid())
 	{
@@ -47,7 +41,6 @@ FTranslationDataManager::FTranslationDataManager( const FString& InManifestFileP
 		if (ArchivePtr.IsValid())
 		{
 			TSharedRef< FInternationalizationArchive > Archive = ArchivePtr.ToSharedRef();
-			TArray<FTranslationUnit> TranslationUnits;
 			int32 NumManifestEntriesParsed = 0;
 
 			GWarn->BeginSlowTask(LOCTEXT("LoadingCurrentManifest", "Loading Entries from Current Translation Manifest..."), true);
@@ -56,10 +49,13 @@ FTranslationDataManager::FTranslationDataManager( const FString& InManifestFileP
 			{
 				GWarn->StatusUpdate(NumManifestEntriesParsed, ManifestEntriesCount, FText::Format(LOCTEXT("LoadingCurrentManifestEntries", "Loading Entry {0} of {1} from Current Translation Manifest..."), FText::AsNumber(NumManifestEntriesParsed), FText::AsNumber(ManifestEntriesCount)));
 				const TSharedRef<FManifestEntry> ManifestEntry = ManifestItr.Value();
-				FTranslationUnit TranslationUnit;
-				TranslationUnit.HasBeenReviewed = false;
-				TranslationUnit.Source = ManifestEntry->Source.Text.ReplaceEscapedCharWithChar();
-				TranslationUnit.Namespace = ManifestEntry->Namespace;
+				UTranslationUnit* TranslationUnit = NewObject<UTranslationUnit>();
+				check(TranslationUnit != NULL);
+				// We want Undo/Redo support
+				TranslationUnit->SetFlags(RF_Transactional);
+				TranslationUnit->HasBeenReviewed = false;
+				TranslationUnit->Source = ManifestEntry->Source.Text.ReplaceEscapedCharWithChar();
+				TranslationUnit->Namespace = ManifestEntry->Namespace;
 
 				for(auto ContextIter( ManifestEntry->Contexts.CreateConstIterator() ); ContextIter; ++ContextIter)
 				{
@@ -69,25 +65,25 @@ FTranslationDataManager::FTranslationDataManager( const FString& InManifestFileP
 					ContextInfo.Context = AContext.SourceLocation;
 					ContextInfo.Key = AContext.Key;
 
-					TranslationUnit.Contexts.Add(ContextInfo);
+					TranslationUnit->Contexts.Add(ContextInfo);
 				}
 
-				TranslationUnits.Add(TranslationUnit);
+				AllTranslations.Add(TranslationUnit);
 			}
 			GWarn->EndSlowTask();
 
 			// Once we have the context information, we can look up history
-			GetHistoryForTranslationUnits(TranslationUnits, ManifestFilePath);
+			GetHistoryForTranslationUnits(AllTranslations, ManifestFilePath);
 
 			GWarn->BeginSlowTask(LOCTEXT("LoadingArchiveEntries", "Loading Entries from Translation Archive..."), true);
-			for (int32 CurrentTranslationUnitIndex = 0; CurrentTranslationUnitIndex < TranslationUnits.Num(); ++CurrentTranslationUnitIndex)
+			for (int32 CurrentTranslationUnitIndex = 0; CurrentTranslationUnitIndex < AllTranslations.Num(); ++CurrentTranslationUnitIndex)
 			{
-				FTranslationUnit& TranslationUnit = TranslationUnits[CurrentTranslationUnitIndex];
+				UTranslationUnit* TranslationUnit = AllTranslations[CurrentTranslationUnitIndex];
 
-				GWarn->StatusUpdate(CurrentTranslationUnitIndex, TranslationUnits.Num(), FText::Format(LOCTEXT("LoadingCurrentArchiveEntries", "Loading Entry {0} of {1} from Translation Archive..."), FText::AsNumber(CurrentTranslationUnitIndex), FText::AsNumber(TranslationUnits.Num())));
+				GWarn->StatusUpdate(CurrentTranslationUnitIndex, AllTranslations.Num(), FText::Format(LOCTEXT("LoadingCurrentArchiveEntries", "Loading Entry {0} of {1} from Translation Archive..."), FText::AsNumber(CurrentTranslationUnitIndex), FText::AsNumber(AllTranslations.Num())));
 
-				const FLocItem SourceSearch(TranslationUnit.Source);
-				TSharedPtr<FArchiveEntry> ArchiveEntry = Archive->FindEntryBySource(TranslationUnit.Namespace, SourceSearch, NULL);
+				const FLocItem SourceSearch(TranslationUnit->Source);
+				TSharedPtr<FArchiveEntry> ArchiveEntry = Archive->FindEntryBySource(TranslationUnit->Namespace, SourceSearch, NULL);
 				if ( ArchiveEntry.IsValid() )
 				{
 					const FString UnescapedTranslatedString = ArchiveEntry->Translation.Text.ReplaceEscapedCharWithChar();
@@ -98,11 +94,11 @@ FTranslationDataManager::FTranslationDataManager( const FString& InManifestFileP
 						int32 MostRecentNonNullTranslationIndex = -1;
 						int32 ContextForRecentTranslation = -1;
 
-						for (int32 ContextIndex = 0; ContextIndex < TranslationUnit.Contexts.Num(); ++ContextIndex)
+						for (int32 ContextIndex = 0; ContextIndex < TranslationUnit->Contexts.Num(); ++ContextIndex)
 						{
-							for (int32 ChangeIndex = 0; ChangeIndex < TranslationUnit.Contexts[ContextIndex].Changes.Num(); ++ChangeIndex)
+							for (int32 ChangeIndex = 0; ChangeIndex < TranslationUnit->Contexts[ContextIndex].Changes.Num(); ++ChangeIndex)
 							{
-								if (!(TranslationUnit.Contexts[ContextIndex].Changes[ChangeIndex].Translation.IsEmpty()))
+								if (!(TranslationUnit->Contexts[ContextIndex].Changes[ChangeIndex].Translation.IsEmpty()))
 								{
 									bHasTranslationHistory = true;
 									MostRecentNonNullTranslationIndex = ChangeIndex;
@@ -121,19 +117,19 @@ FTranslationDataManager::FTranslationDataManager( const FString& InManifestFileP
 						if (bHasTranslationHistory)
 						{
 							// Offer the most recent translation (for the first context in the list) as a suggestion or starting point (not saved unless user checks "Has Been Reviewed")
-							TranslationUnit.Translation = TranslationUnit.Contexts[ContextForRecentTranslation].Changes[MostRecentNonNullTranslationIndex].Translation;
-							TranslationData->Review.Add(TranslationUnit);
+							TranslationUnit->Translation = TranslationUnit->Contexts[ContextForRecentTranslation].Changes[MostRecentNonNullTranslationIndex].Translation;
+							Review.Add(TranslationUnit);
 						}
 						else
 						{
-							TranslationData->Untranslated.Add(TranslationUnit);
+							Untranslated.Add(TranslationUnit);
 						}
 					}
 					else
 					{
-						TranslationUnit.Translation = UnescapedTranslatedString;
-						TranslationUnit.HasBeenReviewed = true;
-						TranslationData->Complete.Add(TranslationUnit);
+						TranslationUnit->Translation = UnescapedTranslatedString;
+						TranslationUnit->HasBeenReviewed = true;
+						Complete.Add(TranslationUnit);
 					}
 				}
 			}
@@ -229,38 +225,38 @@ void FTranslationDataManager::WriteTranslationData()
 
 	bool bNeedsWrite = false;
 
-	for (FTranslationUnit TranslationUnit : TranslationData->Untranslated)
+	for (UTranslationUnit* TranslationUnit : Untranslated)
 	{
-		const FLocItem SearchSource(TranslationUnit.Source);
-		FString OldTranslation = Archive->FindEntryBySource(TranslationUnit.Namespace, SearchSource, NULL)->Translation.Text;
-		FString TranslationToWrite = TranslationUnit.Translation.ReplaceCharWithEscapedChar();
+		const FLocItem SearchSource(TranslationUnit->Source);
+		FString OldTranslation = Archive->FindEntryBySource(TranslationUnit->Namespace, SearchSource, NULL)->Translation.Text;
+		FString TranslationToWrite = TranslationUnit->Translation.ReplaceCharWithEscapedChar();
 		if ( !TranslationToWrite.Equals(OldTranslation) )
 		{
-			Archive->SetTranslation(TranslationUnit.Namespace, TranslationUnit.Source, TranslationToWrite, NULL);
+			Archive->SetTranslation(TranslationUnit->Namespace, TranslationUnit->Source, TranslationToWrite, NULL);
 			bNeedsWrite = true;
 		}
 	}
 
-	for (FTranslationUnit TranslationUnit : TranslationData->Review)
+	for (UTranslationUnit* TranslationUnit : Review)
 	{
-		const FLocItem SearchSource(TranslationUnit.Source);
-		FString OldTranslation = Archive->FindEntryBySource(TranslationUnit.Namespace, SearchSource, NULL)->Translation.Text;
-		FString TranslationToWrite = TranslationUnit.Translation.ReplaceCharWithEscapedChar();
-		if ( TranslationUnit.HasBeenReviewed && !TranslationToWrite.Equals(OldTranslation) )
+		const FLocItem SearchSource(TranslationUnit->Source);
+		FString OldTranslation = Archive->FindEntryBySource(TranslationUnit->Namespace, SearchSource, NULL)->Translation.Text;
+		FString TranslationToWrite = TranslationUnit->Translation.ReplaceCharWithEscapedChar();
+		if ( TranslationUnit->HasBeenReviewed && !TranslationToWrite.Equals(OldTranslation) )
 		{
-			Archive->SetTranslation(TranslationUnit.Namespace, TranslationUnit.Source, TranslationToWrite, NULL);
+			Archive->SetTranslation(TranslationUnit->Namespace, TranslationUnit->Source, TranslationToWrite, NULL);
 			bNeedsWrite = true;
 		}
 	}
 
-	for (FTranslationUnit TranslationUnit : TranslationData->Complete)
+	for (UTranslationUnit* TranslationUnit : Complete)
 	{
-		const FLocItem SearchSource(TranslationUnit.Source);
-		FString OldTranslation = Archive->FindEntryBySource(TranslationUnit.Namespace, SearchSource, NULL)->Translation.Text;
-		FString TranslationToWrite = TranslationUnit.Translation.ReplaceCharWithEscapedChar();
+		const FLocItem SearchSource(TranslationUnit->Source);
+		FString OldTranslation = Archive->FindEntryBySource(TranslationUnit->Namespace, SearchSource, NULL)->Translation.Text;
+		FString TranslationToWrite = TranslationUnit->Translation.ReplaceCharWithEscapedChar();
 		if ( !TranslationToWrite.Equals(OldTranslation) )
 		{
-			Archive->SetTranslation(TranslationUnit.Namespace, TranslationUnit.Source, TranslationToWrite, NULL);
+			Archive->SetTranslation(TranslationUnit->Namespace, TranslationUnit->Source, TranslationToWrite, NULL);
 			bNeedsWrite = true;
 		}
 	}
@@ -381,7 +377,7 @@ bool FTranslationDataManager::WriteJSONToTextFile(TSharedRef<FJsonObject>& Outpu
 	return CheckoutAndSaveWasSuccessful;
 }
 
-void FTranslationDataManager::GetHistoryForTranslationUnits( TArray<FTranslationUnit>& TranslationUnits, const FString& ManifestFilePath )
+void FTranslationDataManager::GetHistoryForTranslationUnits( TArray<UTranslationUnit*>& TranslationUnits, const FString& ManifestFilePath )
 {
 	GWarn->BeginSlowTask(LOCTEXT("LoadingSourceControlHistory", "Loading Translation History from Source Control..."), true);
 
@@ -451,11 +447,11 @@ void FTranslationDataManager::GetHistoryForTranslationUnits( TArray<FTranslation
 				{
 					TSharedRef< FInternationalizationManifest > OldManifest = OldManifestPtr.ToSharedRef();
 
-					for (FTranslationUnit& TranslationUnit : TranslationUnits)
+					for (UTranslationUnit* TranslationUnit : TranslationUnits)
 					{
-						if(TranslationUnit.Contexts.Num() > 0)
+						if(TranslationUnit->Contexts.Num() > 0)
 						{
-							for (FTranslationContextInfo& ContextInfo : TranslationUnit.Contexts)
+							for (FTranslationContextInfo& ContextInfo : TranslationUnit->Contexts)
 							{
 								FString PreviousSourceText = "";
 
@@ -467,7 +463,7 @@ void FTranslationDataManager::GetHistoryForTranslationUnits( TArray<FTranslation
 
 								FContext SearchContext;
 								SearchContext.Key = ContextInfo.Key;
-								TSharedPtr< FManifestEntry > OldManifestEntryPtr = OldManifest->FindEntryByContext(TranslationUnit.Namespace, SearchContext);
+								TSharedPtr< FManifestEntry > OldManifestEntryPtr = OldManifest->FindEntryByContext(TranslationUnit->Namespace, SearchContext);
 								if (!OldManifestEntryPtr.IsValid())
 								{
 									// If this version of the manifest didn't know anything about this string, move onto the next
diff --git a/Engine/Source/Editor/TranslationEditor/Private/TranslationDataManager.h b/Engine/Source/Editor/TranslationEditor/Private/TranslationDataManager.h
--- a/Engine/Source/Editor/TranslationEditor/Private/TranslationDataManager.h
+++ b/Engine/Source/Editor/TranslationEditor/Private/TranslationDataManager.h
@@ -2,24 +2,39 @@
 
 #pragma once
 
-class UTranslationDataObject;
 class FTranslationEditor;
 class FInternationalizationArchive;
 class FInternationalizationManifest;
+class UTranslationUnit;
 
 #include "InternationalizationArchiveJsonSerializer.h"
 #include "InternationalizationManifestJsonSerializer.h"
-#include "TranslationDataObject.h"
 
 class FTranslationDataManager : public TSharedFromThis<FTranslationDataManager>
 {
 
 public:
+
 	FTranslationDataManager( const FString& InManifestFilePath, const FString& InArchiveFilePath);
 
-	UTranslationDataObject* GetTranslationDataObject() 
+	TArray<UTranslationUnit*>& GetAllTranslationsArray()
+	{
+		return AllTranslations;
+	}
+
+	TArray<UTranslationUnit*>& GetUntranslatedArray()
+	{
+		return Untranslated;
+	}
+
+	TArray<UTranslationUnit*>& GetReviewArray()
+	{
+		return Review;
+	}
+
+	TArray<UTranslationUnit*>& GetCompleteArray()
 	{
-		return TranslationData;
+		return Complete;
 	}
 	
 	/** Write the translation data in memory out to .archive file (check out the .archive file first if necessary) */
@@ -44,10 +59,13 @@ private:
 	bool WriteJSONToTextFile( TSharedRef<FJsonObject>& Output, const FString& Filename );
 
 	/** Get the history data for a given translation unit */
-	void GetHistoryForTranslationUnits( TArray<FTranslationUnit>& TranslationUnits, const FString& ManifestFilePath );
+	void GetHistoryForTranslationUnits( TArray<UTranslationUnit*>& TranslationUnits, const FString& ManifestFilePath );
 
-	/** UObject containing our translation information */
-	UTranslationDataObject* TranslationData;
+	// Arrays containing the translation data
+	TArray<UTranslationUnit*> AllTranslations;
+	TArray<UTranslationUnit*> Untranslated;
+	TArray<UTranslationUnit*> Review;
+	TArray<UTranslationUnit*> Complete;
 
 	/** Serializes and deserializes our Archive */
 	FInternationalizationArchiveJsonSerializer ArchiveSerializer;
diff --git a/Engine/Source/Editor/TranslationEditor/Private/TranslationDataObject.cpp b/Engine/Source/Editor/TranslationEditor/Private/TranslationDataObject.cpp
deleted file mode 100644
diff --git a/Engine/Source/Editor/TranslationEditor/Private/TranslationDataObject.h b/Engine/Source/Editor/TranslationEditor/Private/TranslationDataObject.h
deleted file mode 100644
diff --git a/Engine/Source/Editor/TranslationEditor/Private/TranslationEditor.cpp b/Engine/Source/Editor/TranslationEditor/Private/TranslationEditor.cpp
--- a/Engine/Source/Editor/TranslationEditor/Private/TranslationEditor.cpp
+++ b/Engine/Source/Editor/TranslationEditor/Private/TranslationEditor.cpp
@@ -15,7 +15,7 @@
 #include "Editor/PropertyEditor/Public/PropertyPath.h"
 #include "CustomFontColumn.h"
 
-#include "TranslationDataObject.h"
+#include "TranslationUnit.h"
 
 #include "DesktopPlatformModule.h"
 
@@ -82,11 +82,8 @@ void FTranslationEditor::UnregisterTabSpawners(const TSharedRef<class FTabManage
 	TabManager->UnregisterTabSpawner( HistoryTabId );
 }
 
-void FTranslationEditor::InitTranslationEditor( const EToolkitMode::Type Mode, const TSharedPtr< class IToolkitHost >& InitToolkitHost, UTranslationDataObject* TranslationDataToEdit )
+void FTranslationEditor::InitTranslationEditor( const EToolkitMode::Type Mode, const TSharedPtr< class IToolkitHost >& InitToolkitHost )
 {	
-	TranslationData = TranslationDataToEdit;
-	PropertyTableObjects.Add(TranslationData);
-
 	TSharedRef<FTabManager::FLayout> StandaloneDefaultLayout = FTabManager::NewLayout( "Standalone_TranslationEditor_Layout" )
 	->AddArea
 	(
@@ -138,7 +135,17 @@ void FTranslationEditor::InitTranslationEditor( const EToolkitMode::Type Mode, c
 
 	const bool bCreateDefaultStandaloneMenu = true;
 	const bool bCreateDefaultToolbar = true;
-	FAssetEditorToolkit::InitAssetEditor( Mode, InitToolkitHost, FTranslationEditorModule::TranslationEditorAppIdentifier, StandaloneDefaultLayout, bCreateDefaultStandaloneMenu, bCreateDefaultToolbar, TranslationData );
+	// Need editing object to not be null
+	UTranslationUnit* EditingObject;
+	if (DataManager->GetAllTranslationsArray().Num() > 0)
+	{
+		EditingObject = DataManager->GetAllTranslationsArray()[0];
+	}
+	else
+	{
+		EditingObject = NewObject<UTranslationUnit>();
+	}
+	FAssetEditorToolkit::InitAssetEditor(Mode, InitToolkitHost, FTranslationEditorModule::TranslationEditorAppIdentifier, StandaloneDefaultLayout, bCreateDefaultStandaloneMenu, bCreateDefaultToolbar, EditingObject);
 	
 	FTranslationEditorModule& TranslationEditorModule = FModuleManager::LoadModuleChecked<FTranslationEditorModule>( "TranslationEditor" );
 	AddMenuExtender(TranslationEditorModule.GetMenuExtensibilityManager()->GetAllExtenders(GetToolkitCommands(), GetEditingObjects()));
@@ -207,8 +214,8 @@ TSharedRef<SDockTab> FTranslationEditor::SpawnTab_Untranslated( const FSpawnTabA
 
 	FPropertyEditorModule& PropertyEditorModule = FModuleManager::GetModuleChecked<FPropertyEditorModule>( "PropertyEditor" );
 
-	UProperty* SourceProperty = FindField<UProperty>( FTranslationUnit::StaticStruct(), "Source");
-	UProperty* TranslationProperty = FindField<UProperty>( FTranslationUnit::StaticStruct(), "Translation");
+	UProperty* SourceProperty = FindField<UProperty>( UTranslationUnit::StaticClass(), "Source");
+	UProperty* TranslationProperty = FindField<UProperty>( UTranslationUnit::StaticClass(), "Translation");
 
 	// create empty property table
 	UntranslatedPropertyTable = PropertyEditorModule.CreatePropertyTable();
@@ -225,10 +232,7 @@ TSharedRef<SDockTab> FTranslationEditor::SpawnTab_Untranslated( const FSpawnTabA
 	CustomColumns.Add( SourceColumn );
 	CustomColumns.Add(TranslationColumn);
 
-	UntranslatedPropertyTable->SetObjects(PropertyTableObjects);
-
-	// Build the Path to the data we want to show
-	UntranslatedPropertyTable->SetRootPath(FPropertyPath::Create(FindField<UArrayProperty>( UTranslationDataObject::StaticClass(), "Untranslated" )));
+	UntranslatedPropertyTable->SetObjects((TArray<UObject*>&)DataManager->GetUntranslatedArray());
 
 	// Add the columns we want to display
 	UntranslatedPropertyTable->AddColumn((TWeakObjectPtr<UProperty>)SourceProperty);
@@ -266,8 +270,8 @@ TSharedRef<SDockTab> FTranslationEditor::SpawnTab_Review( const FSpawnTabArgs& A
 {
 	check( Args.GetTabId().TabType == ReviewTabId );
 
-	UProperty* SourceProperty = FindField<UProperty>( FTranslationUnit::StaticStruct(), "Source");
-	UProperty* TranslationProperty = FindField<UProperty>( FTranslationUnit::StaticStruct(), "Translation");
+	UProperty* SourceProperty = FindField<UProperty>( UTranslationUnit::StaticClass(), "Source");
+	UProperty* TranslationProperty = FindField<UProperty>( UTranslationUnit::StaticClass(), "Translation");
 
 	FPropertyEditorModule& PropertyEditorModule = FModuleManager::GetModuleChecked<FPropertyEditorModule>( "PropertyEditor" );
 
@@ -286,15 +290,12 @@ TSharedRef<SDockTab> FTranslationEditor::SpawnTab_Review( const FSpawnTabArgs& A
 	CustomColumns.Add( SourceColumn );
 	CustomColumns.Add( TranslationColumn );
 
-	ReviewPropertyTable->SetObjects(PropertyTableObjects);
-
-	// Build the Path to the data we want to show
-	ReviewPropertyTable->SetRootPath(FPropertyPath::Create(FindField<UArrayProperty>( UTranslationDataObject::StaticClass(), "Review" )));
+	ReviewPropertyTable->SetObjects((TArray<UObject*>&)DataManager->GetReviewArray());
 
 	// Add the columns we want to display
-	ReviewPropertyTable->AddColumn((TWeakObjectPtr<UProperty>)FindField<UProperty>( FTranslationUnit::StaticStruct(), "Source"));
-	ReviewPropertyTable->AddColumn((TWeakObjectPtr<UProperty>)FindField<UProperty>( FTranslationUnit::StaticStruct(), "Translation"));
-	ReviewPropertyTable->AddColumn((TWeakObjectPtr<UProperty>)FindField<UProperty>( FTranslationUnit::StaticStruct(), "HasBeenReviewed"));
+	ReviewPropertyTable->AddColumn((TWeakObjectPtr<UProperty>)FindField<UProperty>( UTranslationUnit::StaticClass(), "Source"));
+	ReviewPropertyTable->AddColumn((TWeakObjectPtr<UProperty>)FindField<UProperty>( UTranslationUnit::StaticClass(), "Translation"));
+	ReviewPropertyTable->AddColumn((TWeakObjectPtr<UProperty>)FindField<UProperty>( UTranslationUnit::StaticClass(), "HasBeenReviewed"));
 
 	TArray<TSharedRef<IPropertyTableColumn>> Columns = ReviewPropertyTable->GetColumns();
 	for (TSharedRef<IPropertyTableColumn> Column : Columns)
@@ -334,8 +335,8 @@ TSharedRef<SDockTab> FTranslationEditor::SpawnTab_Completed( const FSpawnTabArgs
 {
 	check( Args.GetTabId().TabType == CompletedTabId );
 
-	UProperty* SourceProperty = FindField<UProperty>( FTranslationUnit::StaticStruct(), "Source");
-	UProperty* TranslationProperty = FindField<UProperty>( FTranslationUnit::StaticStruct(), "Translation");
+	UProperty* SourceProperty = FindField<UProperty>( UTranslationUnit::StaticClass(), "Source");
+	UProperty* TranslationProperty = FindField<UProperty>( UTranslationUnit::StaticClass(), "Translation");
 
 	FPropertyEditorModule& PropertyEditorModule = FModuleManager::GetModuleChecked<FPropertyEditorModule>( "PropertyEditor" );
 
@@ -354,14 +355,11 @@ TSharedRef<SDockTab> FTranslationEditor::SpawnTab_Completed( const FSpawnTabArgs
 	CustomColumns.Add( SourceColumn );
 	CustomColumns.Add( TranslationColumn );
 
-	CompletedPropertyTable->SetObjects(PropertyTableObjects);
-
-	// Build the Path to the data we want to show
-	CompletedPropertyTable->SetRootPath(FPropertyPath::Create(FindField<UArrayProperty>( UTranslationDataObject::StaticClass(), "Complete" )));
+	CompletedPropertyTable->SetObjects((TArray<UObject*>&)DataManager->GetCompleteArray());
 
 	// Add the columns we want to display
-	CompletedPropertyTable->AddColumn((TWeakObjectPtr<UProperty>)FindField<UProperty>( FTranslationUnit::StaticStruct(), "Source"));
-	CompletedPropertyTable->AddColumn((TWeakObjectPtr<UProperty>)FindField<UProperty>( FTranslationUnit::StaticStruct(), "Translation"));
+	CompletedPropertyTable->AddColumn((TWeakObjectPtr<UProperty>)FindField<UProperty>( UTranslationUnit::StaticClass(), "Source"));
+	CompletedPropertyTable->AddColumn((TWeakObjectPtr<UProperty>)FindField<UProperty>( UTranslationUnit::StaticClass(), "Translation"));
 
 	// Freeze columns, don't want user to remove them
 	TArray<TSharedRef<IPropertyTableColumn>> Columns = CompletedPropertyTable->GetColumns();
@@ -431,19 +429,19 @@ TSharedRef<SDockTab> FTranslationEditor::SpawnTab_Context( const FSpawnTabArgs&
 	ContextPropertyTable->SetShowObjectName( false );
 	ContextPropertyTable->OnSelectionChanged()->AddSP( this, &FTranslationEditor::UpdateContextSelection );
 
-	ContextPropertyTable->SetObjects(PropertyTableObjects);
+	if (DataManager->GetAllTranslationsArray().Num() > 0)
+	{
+		TArray<UObject*> Objects;
+		Objects.Add(DataManager->GetAllTranslationsArray()[0]);
+		ContextPropertyTable->SetObjects(Objects);
+	}
 
 	// Build the Path to the data we want to show
-	UArrayProperty* Prop = FindField<UArrayProperty>( UTranslationDataObject::StaticClass(), "Untranslated" );
-	TSharedRef<FPropertyPath> Path = FPropertyPath::Create(Prop);
-	FPropertyInfo PropInfo;
-	PropInfo.Property = Prop->Inner;
-	PropInfo.ArrayIndex = 0;
-	Path = Path->ExtendPath(PropInfo);
-	UProperty* ContextProp = FindField<UProperty>( FTranslationUnit::StaticStruct(), "Contexts" );
+	UProperty* ContextProp = FindField<UProperty>( UTranslationUnit::StaticClass(), "Contexts" );
 	FPropertyInfo ContextPropInfo;
 	ContextPropInfo.Property = ContextProp;
 	ContextPropInfo.ArrayIndex = INDEX_NONE;
+	TSharedRef<FPropertyPath> Path = FPropertyPath::CreateEmpty();
 	Path = Path->ExtendPath(ContextPropInfo);
 	ContextPropertyTable->SetRootPath(Path);
 
@@ -513,17 +511,16 @@ TSharedRef<SDockTab> FTranslationEditor::SpawnTab_History( const FSpawnTabArgs&
 	CustomColumns.Add( SourceColumn );
 	CustomColumns.Add( TranslationColumn );
 
-	HistoryPropertyTable->SetObjects(PropertyTableObjects);
+	if (DataManager->GetAllTranslationsArray().Num() > 0)
+	{
+		TArray<UObject*> Objects;
+		Objects.Add(DataManager->GetAllTranslationsArray()[0]);
+		HistoryPropertyTable->SetObjects(Objects);
+	}
 
 	// Build the Path to the data we want to show
-	UArrayProperty* ReviewProp = FindField<UArrayProperty>( UTranslationDataObject::StaticClass(), "Review" );
-	TSharedRef<FPropertyPath> Path = FPropertyPath::Create(ReviewProp);
-	FPropertyInfo PropInfo;
-	PropInfo.Property = ReviewProp->Inner;
-	PropInfo.ArrayIndex = 0;
-	Path = Path->ExtendPath(PropInfo);
-
-	UArrayProperty* ContextsProp = FindField<UArrayProperty>( FTranslationUnit::StaticStruct(), "Contexts" );
+	TSharedRef<FPropertyPath> Path = FPropertyPath::CreateEmpty();
+	UArrayProperty* ContextsProp = FindField<UArrayProperty>( UTranslationUnit::StaticClass(), "Contexts" );
 	Path = Path->ExtendPath(FPropertyPath::Create(ContextsProp));
 	FPropertyInfo ContextsPropInfo;
 	ContextsPropInfo.Property = ContextsProp->Inner;
@@ -657,7 +654,6 @@ void FTranslationEditor::RefreshUI()
 bool FTranslationEditor::OpenFontPicker( const FString DefaultFile, FString& OutFile )
 {
 	const FString FontFileDescription = LOCTEXT( "FontFileDescription", "Font File" ).ToString();
-	//TODO: support more than one filetype (ttf and otf are supported)
 	const FString FontFileExtension = TEXT("*.ttf;*.otf");
 	const FString FileTypes = FString::Printf( TEXT("%s (%s)|%s"), *FontFileDescription, *FontFileExtension, *FontFileExtension );
 
@@ -701,8 +697,6 @@ bool FTranslationEditor::OpenFontPicker( const FString DefaultFile, FString& Out
 void FTranslationEditor::UpdateTranslationUnitSelection()
 {
 	TSet<TSharedRef<IPropertyTableRow>> SelectedRows;
-	TSharedRef<FPropertyPath> InitialPath = FPropertyPath::CreateEmpty(); // Dummy Path for now
-	UProperty* PropertyToFind = NULL;
 
 	// Different Selection based on which tab is in the foreground
 	TSharedPtr<SDockTab> UntranslatedTabSharedPtr = UntranslatedTab.Pin();
@@ -711,20 +705,14 @@ void FTranslationEditor::UpdateTranslationUnitSelection()
 	if (UntranslatedTabSharedPtr.IsValid() && UntranslatedTabSharedPtr->IsForeground() && UntranslatedPropertyTable.IsValid())
 	{
 		SelectedRows = UntranslatedPropertyTable->GetSelectedRows();
-		InitialPath = UntranslatedPropertyTable->GetRootPath();
-		PropertyToFind = FindField<UProperty>( UTranslationDataObject::StaticClass(), "Untranslated");
 	}
 	else if (ReviewTab.IsValid() && ReviewTabSharedPtr->IsForeground() && ReviewPropertyTable.IsValid())
 	{
 		SelectedRows = ReviewPropertyTable->GetSelectedRows();
-		InitialPath = ReviewPropertyTable->GetRootPath();
-		PropertyToFind = FindField<UProperty>( UTranslationDataObject::StaticClass(), "Review");
 	}
 	else if (CompletedTab.IsValid() && CompletedTabSharedPtr->IsForeground() && CompletedPropertyTable.IsValid())
 	{
 		SelectedRows = CompletedPropertyTable->GetSelectedRows();
-		InitialPath = CompletedPropertyTable->GetRootPath();
-		PropertyToFind = FindField<UProperty>( UTranslationDataObject::StaticClass(), "Complete");
 	}
 
 	// Can only really handle single selection
@@ -732,62 +720,55 @@ void FTranslationEditor::UpdateTranslationUnitSelection()
 	{
 		TSharedRef<IPropertyTableRow> SelectedRow = *(SelectedRows.CreateConstIterator());
 		TSharedRef<FPropertyPath> PartialPath = SelectedRow->GetPartialPath();
-		
-		TArray<FTranslationUnit>& TranslationUnitArray = *(PropertyToFind->ContainerPtrToValuePtr<TArray<FTranslationUnit>>(TranslationData));
-		FTranslationUnit& SelectedTranslationUnit = TranslationUnitArray[PartialPath->GetLeafMostProperty().ArrayIndex];
-		PreviewTextBlock->SetText(TranslationUnitArray[PartialPath->GetLeafMostProperty().ArrayIndex].Translation);
-		NamespaceTextBlock->SetText( FText::Format(LOCTEXT("TranslationNamespace", "Namespace: {0}"), FText::FromString(TranslationUnitArray[PartialPath->GetLeafMostProperty().ArrayIndex].Namespace) ) );
-
-		// Add the ContextPropertyTable-specific path
-		UArrayProperty* ContextArrayProp = FindField<UArrayProperty>( FTranslationUnit::StaticStruct(), "Contexts" );
-		FPropertyInfo ContextArrayPropInfo;
-		ContextArrayPropInfo.Property = ContextArrayProp;
-		ContextArrayPropInfo.ArrayIndex = INDEX_NONE;
-		TSharedRef<FPropertyPath> ContextPath = InitialPath;
-		ContextPath = ContextPath->ExtendPath(PartialPath);
-		ContextPath = ContextPath->ExtendPath(ContextArrayPropInfo);
-		
-		if (ContextPropertyTable.IsValid())
-		{
-			ContextPropertyTable->SetRootPath(ContextPath);
-
-			// Need to re-add the columns we want to display
-			ContextPropertyTable->AddColumn((TWeakObjectPtr<UProperty>)FindField<UProperty>(FTranslationContextInfo::StaticStruct(), "Key"));
-			ContextPropertyTable->AddColumn((TWeakObjectPtr<UProperty>)FindField<UProperty>(FTranslationContextInfo::StaticStruct(), "Context"));
 
-			TArray<TSharedRef<IPropertyTableColumn>> Columns = ContextPropertyTable->GetColumns();
-			for (TSharedRef<IPropertyTableColumn> Column : Columns)
+		TWeakObjectPtr<UObject> UObjectWeakPtr = SelectedRow->GetDataSource()->AsUObject();
+		if (UObjectWeakPtr.IsValid())
+		{
+			UObject* UObjectPtr = UObjectWeakPtr.Get();
+			if (UObjectPtr != NULL)
 			{
-				Column->SetFrozen(true);
-			}
-		}
+				UTranslationUnit* SelectedTranslationUnit = (UTranslationUnit*)UObjectPtr;
+				if (SelectedTranslationUnit != NULL)
+				{
+					PreviewTextBlock->SetText(SelectedTranslationUnit->Translation);
+					NamespaceTextBlock->SetText(FText::Format(LOCTEXT("TranslationNamespace", "Namespace: {0}"), FText::FromString(SelectedTranslationUnit->Namespace)));
+
+					// Add the ContextPropertyTable-specific path
+					UArrayProperty* ContextArrayProp = FindField<UArrayProperty>(UTranslationUnit::StaticClass(), "Contexts");
+					FPropertyInfo ContextArrayPropInfo;
+					ContextArrayPropInfo.Property = ContextArrayProp;
+					ContextArrayPropInfo.ArrayIndex = INDEX_NONE;
+					TSharedRef<FPropertyPath> ContextPath = FPropertyPath::CreateEmpty();
+					ContextPath = ContextPath->ExtendPath(PartialPath);
+					ContextPath = ContextPath->ExtendPath(ContextArrayPropInfo);
+
+					if (ContextPropertyTable.IsValid())
+					{
+						TArray<UObject*> ObjectArray;
+						ObjectArray.Add(SelectedTranslationUnit);
+						ContextPropertyTable->SetObjects(ObjectArray);
+						ContextPropertyTable->SetRootPath(ContextPath);
 
-		// Add the HistoryPropertyTable-specific path
-		TSharedRef<FPropertyPath> HistoryPath = ContextPath;
-		FPropertyInfo ContextPropInfo;
-		ContextPropInfo.Property = ContextArrayProp->Inner;
-		ContextPropInfo.ArrayIndex = 0;	// Just show history for the first context until the user selects something else
-		HistoryPath = HistoryPath->ExtendPath(ContextPropInfo);
-		UArrayProperty* ChangesProp = FindField<UArrayProperty>( FTranslationContextInfo::StaticStruct(), "Changes" );
-		FPropertyInfo ChangesPropInfo;
-		ChangesPropInfo.Property = ChangesProp;
-		ChangesPropInfo.ArrayIndex = INDEX_NONE;
-		HistoryPath = HistoryPath->ExtendPath(ChangesPropInfo);
-
-		if (HistoryPropertyTable.IsValid())
-		{
-			HistoryPropertyTable->SetRootPath(HistoryPath);
+						// Need to re-add the columns we want to display
+						ContextPropertyTable->AddColumn((TWeakObjectPtr<UProperty>)FindField<UProperty>(FTranslationContextInfo::StaticStruct(), "Key"));
+						ContextPropertyTable->AddColumn((TWeakObjectPtr<UProperty>)FindField<UProperty>(FTranslationContextInfo::StaticStruct(), "Context"));
+
+						TArray<TSharedRef<IPropertyTableColumn>> Columns = ContextPropertyTable->GetColumns();
+						for (TSharedRef<IPropertyTableColumn> Column : Columns)
+						{
+							Column->SetFrozen(true);
+						}
 
-			// Need to re-add the columns we want to display
-			HistoryPropertyTable->AddColumn((TWeakObjectPtr<UProperty>)FindField<UProperty>(FTranslationChange::StaticStruct(), "Version"));
-			HistoryPropertyTable->AddColumn((TWeakObjectPtr<UProperty>)FindField<UProperty>(FTranslationChange::StaticStruct(), "DateAndTime"));
-			HistoryPropertyTable->AddColumn((TWeakObjectPtr<UProperty>)FindField<UProperty>(FTranslationChange::StaticStruct(), "Source"));
-			HistoryPropertyTable->AddColumn((TWeakObjectPtr<UProperty>)FindField<UProperty>(FTranslationChange::StaticStruct(), "Translation"));
+						TSharedPtr<IPropertyTableCell> ContextToSelectPtr = ContextPropertyTable->GetFirstCellInTable();
 
-			TArray<TSharedRef<IPropertyTableColumn>> Columns = HistoryPropertyTable->GetColumns();
-			for (TSharedRef<IPropertyTableColumn> Column : Columns)
-			{
-				Column->SetFrozen(true);
+						if (ContextToSelectPtr.IsValid())
+						{
+							TSet<TSharedRef<IPropertyTableCell>> CellsToSelect;
+							CellsToSelect.Add(ContextToSelectPtr.ToSharedRef());
+							ContextPropertyTable->SetSelectedCells(CellsToSelect);
+						}
+					}
+				}
 			}
 		}
 	}
@@ -813,62 +794,67 @@ void FTranslationEditor::UpdateContextSelection()
 			TSharedRef<IPropertyTableRow> SelectedRow = *(SelectedRows.CreateConstIterator());
 			TSharedRef<FPropertyPath> PartialPath = SelectedRow->GetPartialPath();
 
-			TArray<FTranslationUnit>& TranslationUnitArray = *(PropertyToFind->ContainerPtrToValuePtr<TArray<FTranslationUnit>>(TranslationData));
-
-			// Index of the next to last property in the list is the index in the array of translation units
-			FTranslationUnit& SelectedTranslationUnit = TranslationUnitArray[InitialPath->GetPropertyInfo(InitialPath->GetNumProperties() - 2).ArrayIndex];
-			// Index of the leaf most property is the context info index we need
-			FTranslationContextInfo& SelectedContextInfo = SelectedTranslationUnit.Contexts[PartialPath->GetLeafMostProperty().ArrayIndex];
-
-			// If this is a translation unit from the review tab and they select a context, possibly update the selected translation with one from that context
-			if (InitialPath->GetRootProperty().Property == FindField<UProperty>(UTranslationDataObject::StaticClass(), "Review"))
+			TWeakObjectPtr<UObject> UObjectWeakPtr = SelectedRow->GetDataSource()->AsUObject();
+			if (UObjectWeakPtr.IsValid())
 			{
-				// Only change the suggested translation if they haven't yet reviewed it
-				if (SelectedTranslationUnit.HasBeenReviewed == false)
+				UObject* UObjectPtr = UObjectWeakPtr.Get();
+				if (UObjectPtr != NULL)
 				{
-					for (int32 ChangeIndex = 0; ChangeIndex < SelectedContextInfo.Changes.Num(); ++ChangeIndex)
+					UTranslationUnit* SelectedTranslationUnit = (UTranslationUnit*)UObjectPtr;
+					if (SelectedTranslationUnit != NULL)
 					{
-						// Find most recent, non-empty translation
-						if (!SelectedContextInfo.Changes[ChangeIndex].Translation.IsEmpty())
+						// Index of the leaf most property is the context info index we need
+						FTranslationContextInfo& SelectedContextInfo = SelectedTranslationUnit->Contexts[PartialPath->GetLeafMostProperty().ArrayIndex];
+
+						// If this is a translation unit from the review tab and they select a context, possibly update the selected translation with one from that context
+						// Only change the suggested translation if they haven't yet reviewed it
+						if (SelectedTranslationUnit->HasBeenReviewed == false)
 						{
-							check(TranslationData != NULL);
-							if (TranslationData != NULL)
+							for (int32 ChangeIndex = 0; ChangeIndex < SelectedContextInfo.Changes.Num(); ++ChangeIndex)
 							{
-								TranslationData->Modify();
-								SelectedTranslationUnit.Translation = SelectedContextInfo.Changes[ChangeIndex].Translation;
-								TranslationData->PostEditChange();
+								// Find most recent, non-empty translation
+								if (!SelectedContextInfo.Changes[ChangeIndex].Translation.IsEmpty() && SelectedTranslationUnit->Translation != SelectedContextInfo.Changes[ChangeIndex].Translation)
+								{
+									SelectedTranslationUnit->Modify();
+									SelectedTranslationUnit->Translation = SelectedContextInfo.Changes[ChangeIndex].Translation;
+									SelectedTranslationUnit->PostEditChange();
+								}
 							}
 						}
-					}
-				}
-			}
 
-			// Add the HistoryPropertyTable-specific path
-			TSharedRef<FPropertyPath> HistoryPath = ContextPropertyTable->GetRootPath();
-			UArrayProperty* ContextArrayProp = FindField<UArrayProperty>(FTranslationUnit::StaticStruct(), "Contexts");
-			FPropertyInfo ContextPropInfo;
-			ContextPropInfo.Property = ContextArrayProp->Inner;
-			ContextPropInfo.ArrayIndex = PartialPath->GetLeafMostProperty().ArrayIndex;
-			HistoryPath = HistoryPath->ExtendPath(ContextPropInfo);
-			UArrayProperty* ChangesProp = FindField<UArrayProperty>(FTranslationContextInfo::StaticStruct(), "Changes");
-			FPropertyInfo ChangesPropInfo;
-			ChangesPropInfo.Property = ChangesProp;
-			ChangesPropInfo.ArrayIndex = INDEX_NONE;
-			HistoryPath = HistoryPath->ExtendPath(ChangesPropInfo);
-			if (HistoryPropertyTable.IsValid())
-			{
-				HistoryPropertyTable->SetRootPath(HistoryPath);
-
-				// Need to re-add the columns we want to display
-				HistoryPropertyTable->AddColumn((TWeakObjectPtr<UProperty>)FindField<UProperty>(FTranslationChange::StaticStruct(), "Version"));
-				HistoryPropertyTable->AddColumn((TWeakObjectPtr<UProperty>)FindField<UProperty>(FTranslationChange::StaticStruct(), "DateAndTime"));
-				HistoryPropertyTable->AddColumn((TWeakObjectPtr<UProperty>)FindField<UProperty>(FTranslationChange::StaticStruct(), "Source"));
-				HistoryPropertyTable->AddColumn((TWeakObjectPtr<UProperty>)FindField<UProperty>(FTranslationChange::StaticStruct(), "Translation"));
 
-				TArray<TSharedRef<IPropertyTableColumn>> Columns = HistoryPropertyTable->GetColumns();
-				for (TSharedRef<IPropertyTableColumn> Column : Columns)
-				{
-					Column->SetFrozen(true);
+						// Add the HistoryPropertyTable-specific path
+						TSharedRef<FPropertyPath> HistoryPath = ContextPropertyTable->GetRootPath();
+						UArrayProperty* ContextArrayProp = FindField<UArrayProperty>(UTranslationUnit::StaticClass(), "Contexts");
+						FPropertyInfo ContextPropInfo;
+						ContextPropInfo.Property = ContextArrayProp->Inner;
+						ContextPropInfo.ArrayIndex = PartialPath->GetLeafMostProperty().ArrayIndex;
+						HistoryPath = HistoryPath->ExtendPath(ContextPropInfo);
+						UArrayProperty* ChangesProp = FindField<UArrayProperty>(FTranslationContextInfo::StaticStruct(), "Changes");
+						FPropertyInfo ChangesPropInfo;
+						ChangesPropInfo.Property = ChangesProp;
+						ChangesPropInfo.ArrayIndex = INDEX_NONE;
+						HistoryPath = HistoryPath->ExtendPath(ChangesPropInfo);
+						if (HistoryPropertyTable.IsValid())
+						{
+							TArray<UObject*> ObjectArray;
+							ObjectArray.Add(SelectedTranslationUnit);
+							HistoryPropertyTable->SetObjects(ObjectArray);
+							HistoryPropertyTable->SetRootPath(HistoryPath);
+
+							// Need to re-add the columns we want to display
+							HistoryPropertyTable->AddColumn((TWeakObjectPtr<UProperty>)FindField<UProperty>(FTranslationChange::StaticStruct(), "Version"));
+							HistoryPropertyTable->AddColumn((TWeakObjectPtr<UProperty>)FindField<UProperty>(FTranslationChange::StaticStruct(), "DateAndTime"));
+							HistoryPropertyTable->AddColumn((TWeakObjectPtr<UProperty>)FindField<UProperty>(FTranslationChange::StaticStruct(), "Source"));
+							HistoryPropertyTable->AddColumn((TWeakObjectPtr<UProperty>)FindField<UProperty>(FTranslationChange::StaticStruct(), "Translation"));
+
+							TArray<TSharedRef<IPropertyTableColumn>> Columns = HistoryPropertyTable->GetColumns();
+							for (TSharedRef<IPropertyTableColumn> Column : Columns)
+							{
+								Column->SetFrozen(true);
+							}
+						}
+					}
 				}
 			}
 		}
diff --git a/Engine/Source/Editor/TranslationEditor/Private/TranslationEditor.h b/Engine/Source/Editor/TranslationEditor/Private/TranslationEditor.h
--- a/Engine/Source/Editor/TranslationEditor/Private/TranslationEditor.h
+++ b/Engine/Source/Editor/TranslationEditor/Private/TranslationEditor.h
@@ -5,6 +5,7 @@
 #include "Toolkits/AssetEditorManager.h"
 #include "ITranslationEditor.h"
 #include "CustomFontColumn.h"
+#include "TranslationUnit.h"
 
 class TRANSLATIONEDITOR_API FTranslationEditor :  public ITranslationEditor
 {
@@ -21,8 +22,11 @@ public:
 		// Some stuff that needs to use the "this" pointer is done in Initialize (because it can't be done in the constructor)
 		TranslationEditor->Initialize();
 
-		// Set up a property changed event to trigger a write of the translation data when TranslationDataObject property changes
-		DataManager->GetTranslationDataObject()->OnPropertyChanged().Add(UTranslationDataObject::FTranslationDataPropertyChangedEvent::FDelegate::CreateSP(DataManager, &FTranslationDataManager::HandlePropertyChanged));
+		for (UTranslationUnit* TranslationUnit : DataManager->GetAllTranslationsArray())
+		{
+			// Set up a property changed event to trigger a write of the translation data when TranslationUnit property changes
+			TranslationUnit->OnPropertyChanged().Add(UTranslationUnit::FTranslationUnitPropertyChangedEvent::FDelegate::CreateSP(DataManager, &FTranslationDataManager::HandlePropertyChanged));
+		}
 
 		return TranslationEditor;
 	}
@@ -39,7 +43,7 @@ public:
 	 * @param	InitToolkitHost			When Mode is WorldCentric, this is the level editor instance to spawn this editor within
 	 * @param	ObjectToEdit			The object to edit
 	 */
-	void InitTranslationEditor( const EToolkitMode::Type Mode, const TSharedPtr< class IToolkitHost >& InitToolkitHost, UTranslationDataObject* TranslationDataToEdit );
+	void InitTranslationEditor( const EToolkitMode::Type Mode, const TSharedPtr< class IToolkitHost >& InitToolkitHost );
 
 	/** IToolkit interface */
 	virtual FName GetToolkitFName() const OVERRIDE;
@@ -55,8 +59,7 @@ protected:
 private:
 
 	FTranslationEditor(TSharedRef< FTranslationDataManager > InDataManager, const FString& InManifestFile, const FString& InArchiveFile )
-	: TranslationData(NULL)
-	, DataManager(InDataManager)
+	:  DataManager(InDataManager)
 	, SourceFont(FEditorStyle::GetFontStyle( PropertyTableConstants::NormalFontStyle ))
 	, TranslationTargetFont(FEditorStyle::GetFontStyle( PropertyTableConstants::NormalFontStyle ))
 	, SourceColumn(MakeShareable(new FCustomFontColumn(SourceFont)))
@@ -162,15 +165,9 @@ private:
 	/** The Review Tab */
 	TWeakPtr<SDockTab> CompletedTab;
 
-	/** UObject containing our translation information */
-	UTranslationDataObject* TranslationData;
-
 	/** Manages the reading and writing of data to file */
 	TSharedRef< FTranslationDataManager > DataManager;
 
-	/** Array of Objects for the Property Table */
-	TArray<UObject*>  PropertyTableObjects;
-
 	/** The table of untranslated items */
 	TSharedPtr< class IPropertyTable > UntranslatedPropertyTable;
 	/** The table of translations to review */
diff --git a/Engine/Source/Editor/TranslationEditor/Private/TranslationEditorModule.cpp b/Engine/Source/Editor/TranslationEditor/Private/TranslationEditorModule.cpp
--- a/Engine/Source/Editor/TranslationEditor/Private/TranslationEditorModule.cpp
+++ b/Engine/Source/Editor/TranslationEditor/Private/TranslationEditorModule.cpp
@@ -5,7 +5,6 @@
 #include "TranslationEditor.h"
 //#include "Toolkits/ToolkitManager.h"
 #include "WorkspaceMenuStructureModule.h"
-#include "TranslationDataObject.h"
 #include "MessageLogModule.h"
 
 class FTranslationEditor;
@@ -46,7 +45,7 @@ TSharedRef<FTranslationEditor> FTranslationEditorModule::CreateTranslationEditor
 	GWarn->BeginSlowTask(LOCTEXT("BuildingUserInterface", "Building Translation Editor UI..."), true);
 
 	TSharedRef< FTranslationEditor > NewTranslationEditor( FTranslationEditor::Create(DataManager, ManifestFile, ArchiveFile) );
-	NewTranslationEditor->InitTranslationEditor( EToolkitMode::Standalone, TSharedPtr<IToolkitHost>(), DataManager->GetTranslationDataObject() );
+	NewTranslationEditor->InitTranslationEditor( EToolkitMode::Standalone, TSharedPtr<IToolkitHost>() );
 
 	GWarn->EndSlowTask();
 
diff --git a/Engine/Source/Editor/TranslationEditor/Public/TranslationEditorModule.h b/Engine/Source/Editor/TranslationEditor/Public/TranslationEditorModule.h
--- a/Engine/Source/Editor/TranslationEditor/Public/TranslationEditorModule.h
+++ b/Engine/Source/Editor/TranslationEditor/Public/TranslationEditorModule.h
@@ -7,7 +7,6 @@
 #include "Toolkits/AssetEditorToolkit.h" // For FExtensibilityManager
 //#include "ITranslationEditor.h"
 
-class UTranslationDataObject;
 class FTranslationEditor;
 
 class FTranslationEditorModule : public IModuleInterface,
