Forgot some files with the last checkin

diff --git a/Engine/Source/Editor/TranslationEditor/Private/TranslationUnit.cpp b/Engine/Source/Editor/TranslationEditor/Private/TranslationUnit.cpp
new file mode 100644
--- /dev/null
+++ b/Engine/Source/Editor/TranslationEditor/Private/TranslationUnit.cpp
@@ -0,0 +1,18 @@
+// Copyright 1998-2014 Epic Games, Inc. All Rights Reserved.
+#include "TranslationEditorPrivatePCH.h"
+#include "TranslationEditor.generated.inl"
+
+UTranslationUnit::UTranslationUnit( const class FPostConstructInitializeProperties& PCIP )
+	: Super(PCIP)
+{
+
+}
+
+void UTranslationUnit::PostEditChangeProperty(struct FPropertyChangedEvent& PropertyChangedEvent)
+{
+	Super::PostEditChangeProperty(PropertyChangedEvent);
+
+	const FName Name = (PropertyChangedEvent.Property != nullptr) ? PropertyChangedEvent.Property->GetFName() : NAME_None;
+
+	TranslationUnitPropertyChangedEvent.Broadcast(Name);
+}
\ No newline at end of file
diff --git a/Engine/Source/Editor/TranslationEditor/Private/TranslationUnit.h b/Engine/Source/Editor/TranslationEditor/Private/TranslationUnit.h
new file mode 100644
--- /dev/null
+++ b/Engine/Source/Editor/TranslationEditor/Private/TranslationUnit.h
@@ -0,0 +1,97 @@
+// Copyright 1998-2014 Epic Games, Inc. All Rights Reserved.
+#pragma once
+
+#include "TranslationUnit.generated.h"
+
+USTRUCT()
+struct FTranslationChange
+{
+	GENERATED_USTRUCT_BODY()
+
+public:
+
+	/** The changelist of this change */
+	UPROPERTY(Category=Translation, VisibleAnywhere)
+	FString Version;
+
+	/** Date of this change */
+	UPROPERTY(Category=Translation, VisibleAnywhere)//, meta=(DisplayName = "Date & Time"))
+	FDateTime DateAndTime;
+
+	/** Source at time of this change */
+	UPROPERTY(Category=Translation, VisibleAnywhere)
+	FString Source;
+
+	/** Translation at time of this change */
+	UPROPERTY(Category=Translation, VisibleAnywhere)
+	FString Translation;
+};
+
+USTRUCT()
+struct FTranslationContextInfo
+{
+	GENERATED_USTRUCT_BODY()
+
+public:
+
+	/** The key specified in LOCTEXT */
+	UPROPERTY(Category=Translation, VisibleAnywhere)
+	FString Key;
+
+	/** What file and line this translation is from */
+	UPROPERTY(Category=Translation, VisibleAnywhere)
+	FString Context;
+
+	/** List of previous versions of the source text for this context */
+	UPROPERTY(Category=Translation, VisibleAnywhere)
+	TArray<FTranslationChange> Changes;
+};
+
+UCLASS(hidecategories = Object, MinimalAPI)
+class UTranslationUnit : public UObject
+{
+	GENERATED_UCLASS_BODY()
+
+public:
+
+	/** The localization namespace for this translation */
+	UPROPERTY(Category=Translation, VisibleAnywhere)
+	FString Namespace;
+
+	/** Original text from the source language */
+	UPROPERTY(Category=Translation, VisibleAnywhere)
+	FString Source;
+
+	/** Translations */
+	UPROPERTY(Category=Translation, EditAnywhere)
+	FString Translation;
+
+	/** Contexts the source was found in */
+	UPROPERTY(Category=Translation, VisibleAnywhere)
+	TArray<FTranslationContextInfo> Contexts;
+
+	/** Whether the changes have been reviewed */
+	UPROPERTY(Category=Translation, EditAnywhere)
+	bool HasBeenReviewed;
+
+	/**
+	* Returns an event delegate that is executed when a property has changed.
+	*
+	* @return The delegate.
+	*/
+	DECLARE_EVENT_OneParam(UTranslationUnit, FTranslationUnitPropertyChangedEvent, FName /*PropertyName*/);
+	FTranslationUnitPropertyChangedEvent& OnPropertyChanged() { return TranslationUnitPropertyChangedEvent; }
+
+protected:
+	/**
+	* Called when a property on this object has been modified externally
+	*
+	* @param PropertyThatChanged the property that was modified
+	*/
+	virtual void PostEditChangeProperty(struct FPropertyChangedEvent& PropertyChangedEvent) OVERRIDE;
+
+private:
+
+	// Holds an event delegate that is executed when a property has changed.
+	FTranslationUnitPropertyChangedEvent TranslationUnitPropertyChangedEvent;
+};
\ No newline at end of file
