Translation Editor:
- Added Search functionality to Translation Editor with a new Search tab.

diff --git a/Engine/Source/Editor/TranslationEditor/Private/TranslationDataManager.cpp b/Engine/Source/Editor/TranslationEditor/Private/TranslationDataManager.cpp
--- a/Engine/Source/Editor/TranslationEditor/Private/TranslationDataManager.cpp
+++ b/Engine/Source/Editor/TranslationEditor/Private/TranslationDataManager.cpp
@@ -19,6 +19,7 @@ FTranslationDataManager::FTranslationDataManager( const FString& InManifestFileP
 , ArchiveFilePath(InArchiveFilePath)
 {
 	GWarn->BeginSlowTask(LOCTEXT("LoadingTranslationData", "Loading Translation Data..."), true);
+	TArray<UTranslationUnit*> TranslationUnits;
 
 	ManifestAtHeadRevisionPtr = ReadManifest( ManifestFilePath );
 	if (ManifestAtHeadRevisionPtr.IsValid())
@@ -68,19 +69,20 @@ FTranslationDataManager::FTranslationDataManager( const FString& InManifestFileP
 					TranslationUnit->Contexts.Add(ContextInfo);
 				}
 
-				AllTranslations.Add(TranslationUnit);
+				TranslationUnits.Add(TranslationUnit);
 			}
 			GWarn->EndSlowTask();
 
 			// Once we have the context information, we can look up history
-			GetHistoryForTranslationUnits(AllTranslations, ManifestFilePath);
+			GetHistoryForTranslationUnits(TranslationUnits, ManifestFilePath);
 
 			GWarn->BeginSlowTask(LOCTEXT("LoadingArchiveEntries", "Loading Entries from Translation Archive..."), true);
-			for (int32 CurrentTranslationUnitIndex = 0; CurrentTranslationUnitIndex < AllTranslations.Num(); ++CurrentTranslationUnitIndex)
+			for (int32 CurrentTranslationUnitIndex = 0; CurrentTranslationUnitIndex < TranslationUnits.Num(); ++CurrentTranslationUnitIndex)
 			{
-				UTranslationUnit* TranslationUnit = AllTranslations[CurrentTranslationUnitIndex];
+				UTranslationUnit* TranslationUnit = TranslationUnits[CurrentTranslationUnitIndex];
+				AllTranslations.Add(TranslationUnit);
 
-				GWarn->StatusUpdate(CurrentTranslationUnitIndex, AllTranslations.Num(), FText::Format(LOCTEXT("LoadingCurrentArchiveEntries", "Loading Entry {0} of {1} from Translation Archive..."), FText::AsNumber(CurrentTranslationUnitIndex), FText::AsNumber(AllTranslations.Num())));
+				GWarn->StatusUpdate(CurrentTranslationUnitIndex, TranslationUnits.Num(), FText::Format(LOCTEXT("LoadingCurrentArchiveEntries", "Loading Entry {0} of {1} from Translation Archive..."), FText::AsNumber(CurrentTranslationUnitIndex), FText::AsNumber(TranslationUnits.Num())));
 
 				const FLocItem SourceSearch(TranslationUnit->Source);
 				TSharedPtr<FArchiveEntry> ArchiveEntry = Archive->FindEntryBySource(TranslationUnit->Namespace, SourceSearch, NULL);
@@ -544,4 +546,18 @@ void FTranslationDataManager::PreviewAllTranslationsInEditor()
 	FTextLocalizationManager::Get().RegenerateResources(ConfigFilePath);
 }
 
+void FTranslationDataManager::PopulateSearchResultsUsingFilter(const FString& SearchFilter)
+{
+	SearchResults.Empty();
+
+	for (UTranslationUnit* TranslationUnit : AllTranslations)
+	{
+		if (TranslationUnit->Source.Contains(SearchFilter) || 
+			TranslationUnit->Translation.Contains(SearchFilter))
+		{
+			SearchResults.Add(TranslationUnit);
+		}
+	}
+}
+
 #undef LOCTEXT_NAMESPACE
diff --git a/Engine/Source/Editor/TranslationEditor/Private/TranslationDataManager.h b/Engine/Source/Editor/TranslationEditor/Private/TranslationDataManager.h
--- a/Engine/Source/Editor/TranslationEditor/Private/TranslationDataManager.h
+++ b/Engine/Source/Editor/TranslationEditor/Private/TranslationDataManager.h
@@ -36,6 +36,11 @@ public:
 	{
 		return Complete;
 	}
+
+	TArray<UTranslationUnit*>& GetSearchResultsArray()
+	{
+		return SearchResults;
+	}
 	
 	/** Write the translation data in memory out to .archive file (check out the .archive file first if necessary) */
 	void WriteTranslationData();
@@ -46,6 +51,9 @@ public:
 	/** Regenerate and reload archives to reflect modifications in the UI */
 	void PreviewAllTranslationsInEditor();
 
+	/** Put items in the Search Array if they match this filter */
+	void PopulateSearchResultsUsingFilter(const FString& SearchFilter);
+
 private:
 	/** Read text file into a JSON file */
 	static TSharedPtr<FJsonObject> ReadJSONTextFile( const FString& InFilePath ) ;
@@ -66,6 +74,7 @@ private:
 	TArray<UTranslationUnit*> Untranslated;
 	TArray<UTranslationUnit*> Review;
 	TArray<UTranslationUnit*> Complete;
+	TArray<UTranslationUnit*> SearchResults;
 
 	/** Serializes and deserializes our Archive */
 	FInternationalizationArchiveJsonSerializer ArchiveSerializer;
diff --git a/Engine/Source/Editor/TranslationEditor/Private/TranslationEditor.cpp b/Engine/Source/Editor/TranslationEditor/Private/TranslationEditor.cpp
--- a/Engine/Source/Editor/TranslationEditor/Private/TranslationEditor.cpp
+++ b/Engine/Source/Editor/TranslationEditor/Private/TranslationEditor.cpp
@@ -21,6 +21,8 @@
 
 #include "IPropertyTableWidgetHandle.h"
 
+#include "SSearchBox.h"
+
 DEFINE_LOG_CATEGORY_STATIC(LocalizationExport, Log, All);
 
 #define LOCTEXT_NAMESPACE "TranslationEditor"
@@ -31,6 +33,7 @@ const FName FTranslationEditor::CompletedTabId( TEXT( "TranslationEditor_Complet
 const FName FTranslationEditor::PreviewTabId( TEXT( "TranslationEditor_Preview" ) );
 const FName FTranslationEditor::ContextTabId( TEXT( "TranslationEditor_Context" ) );
 const FName FTranslationEditor::HistoryTabId( TEXT( "TranslationEditor_History" ) );
+const FName FTranslationEditor::SearchTabId( TEXT( "TranslationEditor_Search" ) );
 
 void FTranslationEditor::Initialize()
 {
@@ -70,6 +73,10 @@ void FTranslationEditor::RegisterTabSpawners(const TSharedRef<class FTabManager>
 	TabManager->RegisterTabSpawner( HistoryTabId, FOnSpawnTab::CreateSP(this, &FTranslationEditor::SpawnTab_History) )
 		.SetDisplayName( LOCTEXT("HistoryTab", "History") )
 		.SetGroup( MenuStructure.GetAssetEditorCategory() );
+
+	TabManager->RegisterTabSpawner(SearchTabId, FOnSpawnTab::CreateSP(this, &FTranslationEditor::SpawnTab_Search))
+		.SetDisplayName(LOCTEXT("SearchTab", "Search"))
+		.SetGroup(MenuStructure.GetAssetEditorCategory());
 }
 
 void FTranslationEditor::UnregisterTabSpawners(const TSharedRef<class FTabManager>& TabManager)
@@ -79,7 +86,8 @@ void FTranslationEditor::UnregisterTabSpawners(const TSharedRef<class FTabManage
 	TabManager->UnregisterTabSpawner( CompletedTabId );
 	TabManager->UnregisterTabSpawner( PreviewTabId );
 	TabManager->UnregisterTabSpawner( ContextTabId );
-	TabManager->UnregisterTabSpawner( HistoryTabId );
+	TabManager->UnregisterTabSpawner(HistoryTabId);
+	TabManager->UnregisterTabSpawner(SearchTabId);
 }
 
 void FTranslationEditor::InitTranslationEditor( const EToolkitMode::Type Mode, const TSharedPtr< class IToolkitHost >& InitToolkitHost )
@@ -95,7 +103,7 @@ void FTranslationEditor::InitTranslationEditor( const EToolkitMode::Type Mode, c
 			->SetSizeCoefficient(0.1f)
 			->SetHideTabWell( true )
 			->AddTab(GetToolbarTabId(), ETabState::OpenedTab)
-		) 
+		)
 		->Split
 		(
 			FTabManager::NewStack()
@@ -104,6 +112,7 @@ void FTranslationEditor::InitTranslationEditor( const EToolkitMode::Type Mode, c
 			->AddTab( UntranslatedTabId, ETabState::OpenedTab )
 			->AddTab( ReviewTabId,  ETabState::OpenedTab )
 			->AddTab( CompletedTabId,  ETabState::OpenedTab )
+			->AddTab( SearchTabId, ETabState::ClosedTab )
 		)
 		->Split
 		(
@@ -389,6 +398,84 @@ TSharedRef<SDockTab> FTranslationEditor::SpawnTab_Completed( const FSpawnTabArgs
 	return NewDockTab;
 }
 
+TSharedRef<SDockTab> FTranslationEditor::SpawnTab_Search(const FSpawnTabArgs& Args)
+{
+	check(Args.GetTabId().TabType == SearchTabId);
+
+	UProperty* SourceProperty = FindField<UProperty>(UTranslationUnit::StaticClass(), "Source");
+	UProperty* TranslationProperty = FindField<UProperty>(UTranslationUnit::StaticClass(), "Translation");
+
+	FPropertyEditorModule& PropertyEditorModule = FModuleManager::GetModuleChecked<FPropertyEditorModule>("PropertyEditor");
+
+	// create empty property table
+	SearchPropertyTable = PropertyEditorModule.CreatePropertyTable();
+	SearchPropertyTable->SetIsUserAllowedToChangeRoot(false);
+	SearchPropertyTable->SetOrientation(EPropertyTableOrientation::AlignPropertiesInColumns);
+	SearchPropertyTable->SetShowRowHeader(true);
+	SearchPropertyTable->SetShowObjectName(false);
+	SearchPropertyTable->OnSelectionChanged()->AddSP(this, &FTranslationEditor::UpdateTranslationUnitSelection);
+
+	// we want to customize some columns
+	TArray< TSharedRef< class IPropertyTableCustomColumn > > CustomColumns;
+	SourceColumn->AddSupportedProperty(SourceProperty);
+	TranslationColumn->AddSupportedProperty(TranslationProperty);
+	CustomColumns.Add(SourceColumn);
+	CustomColumns.Add(TranslationColumn);
+
+	SearchPropertyTable->SetObjects((TArray<UObject*>&)DataManager->GetSearchResultsArray());
+
+	// Add the columns we want to display
+	SearchPropertyTable->AddColumn((TWeakObjectPtr<UProperty>)FindField<UProperty>(UTranslationUnit::StaticClass(), "Source"));
+	SearchPropertyTable->AddColumn((TWeakObjectPtr<UProperty>)FindField<UProperty>(UTranslationUnit::StaticClass(), "Translation"));
+
+	// Freeze columns, don't want user to remove them
+	TArray<TSharedRef<IPropertyTableColumn>> Columns = SearchPropertyTable->GetColumns();
+	for (TSharedRef<IPropertyTableColumn> Column : Columns)
+	{
+		Column->SetFrozen(true);
+	}
+
+	SearchPropertyTableWidgetHandle = PropertyEditorModule.CreatePropertyTableWidgetHandle(SearchPropertyTable.ToSharedRef(), CustomColumns);
+	TSharedRef<SWidget> PropertyTableWidget = SearchPropertyTableWidgetHandle->GetWidget();
+
+	TSharedRef<SDockTab> NewDockTab = SNew(SDockTab)
+		//.Icon(FEditorStyle::GetBrush("TranslationEditor.Tabs.Properties"))
+		.Label(LOCTEXT("SearchTabTitle", "Search"))
+		.TabColorScale(GetTabColorScale())
+		[
+			SNew(SVerticalBox)
+			+ SVerticalBox::Slot()
+			.HAlign(HAlign_Fill)
+			.VAlign(VAlign_Top)
+			.AutoHeight()
+			.Padding(0.0f, 0.0f, 0.0f, 4.0f)
+			[
+				SAssignNew(SearchBox, SSearchBox)
+				.HintText(LOCTEXT("FilterSearch", "Search..."))
+				.ToolTipText(LOCTEXT("FilterSearchHint", "Type here to search").ToString())
+				.OnTextChanged(this, &FTranslationEditor::OnFilterTextChanged)
+				.OnTextCommitted(this, &FTranslationEditor::OnFilterTextCommitted)
+			]
+			+ SVerticalBox::Slot()
+				.HAlign(HAlign_Fill)
+				.VAlign(VAlign_Top)
+				.FillHeight(10.f)
+			[
+				SNew(SBorder)
+				.BorderImage(FEditorStyle::GetBrush("ToolPanel.GroupBorder"))
+				.Padding(0.0f)
+				.VAlign(VAlign_Top)
+				[
+					PropertyTableWidget
+				]
+			]
+		];
+
+	SearchTab = NewDockTab;
+
+	return NewDockTab;
+}
+
 TSharedRef<SDockTab> FTranslationEditor::SpawnTab_Preview( const FSpawnTabArgs& Args )
 {
 	check( Args.GetTabId().TabType == PreviewTabId );
@@ -589,6 +676,10 @@ void FTranslationEditor::MapActions()
 	ToolkitCommands->MapAction(FTranslationEditorCommands::Get().ExportToPortableObjectFormat,
 		FExecuteAction::CreateSP(this, &FTranslationEditor::ExportToPortableObjectFormat_Execute),
 		FCanExecuteAction());
+
+	ToolkitCommands->MapAction(FTranslationEditorCommands::Get().OpenSearchTab,
+		FExecuteAction::CreateSP(this, &FTranslationEditor::OpenSearchTab_Execute),
+		FCanExecuteAction());
 }
 
 void FTranslationEditor::ChangeSourceFont()
@@ -649,6 +740,10 @@ void FTranslationEditor::RefreshUI()
 	{
 		HistoryPropertyTableWidgetHandle->RequestRefresh();
 	}
+	if (SearchPropertyTableWidgetHandle.IsValid())
+	{
+		SearchPropertyTableWidgetHandle->RequestRefresh();
+	}
 }
 
 bool FTranslationEditor::OpenFontPicker( const FString DefaultFile, FString& OutFile )
@@ -702,6 +797,7 @@ void FTranslationEditor::UpdateTranslationUnitSelection()
 	TSharedPtr<SDockTab> UntranslatedTabSharedPtr = UntranslatedTab.Pin();
 	TSharedPtr<SDockTab> ReviewTabSharedPtr = ReviewTab.Pin();
 	TSharedPtr<SDockTab> CompletedTabSharedPtr = CompletedTab.Pin();
+	TSharedPtr<SDockTab> SearchTabSharedPtr = SearchTab.Pin();
 	if (UntranslatedTabSharedPtr.IsValid() && UntranslatedTabSharedPtr->IsForeground() && UntranslatedPropertyTable.IsValid())
 	{
 		SelectedRows = UntranslatedPropertyTable->GetSelectedRows();
@@ -714,6 +810,10 @@ void FTranslationEditor::UpdateTranslationUnitSelection()
 	{
 		SelectedRows = CompletedPropertyTable->GetSelectedRows();
 	}
+	else if (SearchTab.IsValid() && SearchTabSharedPtr->IsForeground() && SearchPropertyTable.IsValid())
+	{
+		SelectedRows = SearchPropertyTable->GetSelectedRows();
+	}
 
 	// Can only really handle single selection
 	if (SelectedRows.Num() == 1)
@@ -956,4 +1056,44 @@ void FTranslationEditor::ExportToPortableObjectFormat_Execute()
 
 }
 
+void FTranslationEditor::OnFilterTextChanged(const FText& InFilterText)
+{
+
+}
+
+void FTranslationEditor::OnFilterTextCommitted(const FText& InFilterText, ETextCommit::Type CommitInfo)
+{
+	const FString InFilterString = InFilterText.ToString();
+
+	if (CommitInfo == ETextCommit::OnEnter)
+	{
+		if (InFilterString != CurrentSearchFilter)
+		{
+			CurrentSearchFilter = InFilterString;
+
+			DataManager->PopulateSearchResultsUsingFilter(InFilterString);
+
+			if (SearchPropertyTable.IsValid())
+			{
+				SearchPropertyTable->SetObjects((TArray<UObject*>&)DataManager->GetSearchResultsArray());
+
+				// Need to re-add the columns we want to display
+				SearchPropertyTable->AddColumn((TWeakObjectPtr<UProperty>)FindField<UProperty>(UTranslationUnit::StaticClass(), "Source"));
+				SearchPropertyTable->AddColumn((TWeakObjectPtr<UProperty>)FindField<UProperty>(UTranslationUnit::StaticClass(), "Translation"));
+
+				TArray<TSharedRef<IPropertyTableColumn>> Columns = SearchPropertyTable->GetColumns();
+				for (TSharedRef<IPropertyTableColumn> Column : Columns)
+				{
+					Column->SetFrozen(true);
+				}
+			}
+		}
+	}
+}
+
+void FTranslationEditor::OpenSearchTab_Execute()
+{
+	TabManager->InvokeTab(SearchTabId);
+}
+
 #undef LOCTEXT_NAMESPACE
diff --git a/Engine/Source/Editor/TranslationEditor/Private/TranslationEditor.h b/Engine/Source/Editor/TranslationEditor/Private/TranslationEditor.h
--- a/Engine/Source/Editor/TranslationEditor/Private/TranslationEditor.h
+++ b/Engine/Source/Editor/TranslationEditor/Private/TranslationEditor.h
@@ -88,6 +88,8 @@ private:
 	TSharedRef<SDockTab> SpawnTab_Context( const FSpawnTabArgs& Args );
 	/**	Spawns the history tab */
 	TSharedRef<SDockTab> SpawnTab_History( const FSpawnTabArgs& Args );
+	/**	Spawns the search tab */
+	TSharedRef<SDockTab> SpawnTab_Search( const FSpawnTabArgs& Args );
 
 	/** Map actions for the UI_COMMANDS */
 	void MapActions();
@@ -145,6 +147,15 @@ private:
 	/** Called when "Export to .PO" is clicked for this asset */
 	void ExportToPortableObjectFormat_Execute();
 
+	/** Open the search tab */
+	void OpenSearchTab_Execute();
+
+	/** Called when the filter text in the SearchBox is changed */
+	void OnFilterTextChanged(const FText& InFilterText);
+
+	/** Called when text is committed to the SearchBox */
+	void OnFilterTextCommitted(const FText& InFilterText, ETextCommit::Type CommitInfo);
+
 	/**	The tab id for the untranslated tab */
 	static const FName UntranslatedTabId;
 	/**	The tab id for the review tab */
@@ -157,6 +168,8 @@ private:
 	static const FName ContextTabId;
 	/**	The tab id for the history tab */
 	static const FName HistoryTabId;
+	/**	The tab id for the search tab */
+	static const FName SearchTabId;
 
 	/** The Untranslated Tab */
 	TWeakPtr<SDockTab> UntranslatedTab;
@@ -164,6 +177,13 @@ private:
 	TWeakPtr<SDockTab> ReviewTab;
 	/** The Review Tab */
 	TWeakPtr<SDockTab> CompletedTab;
+	/** The Search Tab */
+	TWeakPtr<SDockTab> SearchTab;
+
+	/** Search box for searching the source and translation strings */
+	TSharedPtr<SSearchBox> SearchBox;
+	/** Current search filter */
+	FString CurrentSearchFilter;
 
 	/** Manages the reading and writing of data to file */
 	TSharedRef< FTranslationDataManager > DataManager;
@@ -178,6 +198,8 @@ private:
 	TSharedPtr< class IPropertyTable > ContextPropertyTable;
 	/** The table of previous revision information */
 	TSharedPtr< class IPropertyTable > HistoryPropertyTable;
+	/** The table of search results */
+	TSharedPtr< class IPropertyTable > SearchPropertyTable;
 
 	/** The slate widget table of untranslated items */
 	TSharedPtr< class IPropertyTableWidgetHandle > UntranslatedPropertyTableWidgetHandle;
@@ -189,6 +211,8 @@ private:
 	TSharedPtr< class IPropertyTableWidgetHandle > ContextPropertyTableWidgetHandle;
 	/** The slate widget table of previous revision information */
 	TSharedPtr< class IPropertyTableWidgetHandle > HistoryPropertyTableWidgetHandle;
+	/** The slate widget table of search results */
+	TSharedPtr< class IPropertyTableWidgetHandle > SearchPropertyTableWidgetHandle;
 
 	/** Font to use for the source language */
 	FSlateFontInfo SourceFont;
diff --git a/Engine/Source/Editor/TranslationEditor/Private/TranslationEditorMenu.cpp b/Engine/Source/Editor/TranslationEditor/Private/TranslationEditorMenu.cpp
--- a/Engine/Source/Editor/TranslationEditor/Private/TranslationEditorMenu.cpp
+++ b/Engine/Source/Editor/TranslationEditor/Private/TranslationEditorMenu.cpp
@@ -17,6 +17,7 @@ void FTranslationEditorMenu::FillTranslationMenu( FMenuBuilder& MenuBuilder/*, F
 		MenuBuilder.AddMenuEntry( FTranslationEditorCommands::Get().ChangeTranslationTargetFont );
 		MenuBuilder.AddMenuEntry( FTranslationEditorCommands::Get().PreviewAllTranslationsInEditor );
 		MenuBuilder.AddMenuEntry( FTranslationEditorCommands::Get().ExportToPortableObjectFormat );
+		MenuBuilder.AddMenuEntry( FTranslationEditorCommands::Get().OpenSearchTab );
 	}
 	MenuBuilder.EndSection();
 }
@@ -69,6 +70,8 @@ void FTranslationEditorMenu::SetupTranslationEditorToolbar( TSharedPtr< FExtende
 				FTranslationEditorCommands::Get().PreviewAllTranslationsInEditor, "PreviewTranslationsInEditor", TAttribute<FText>(), TAttribute<FText>(), FSlateIcon(FEditorStyle::GetStyleSetName(), "AssetEditor.ReimportAsset")); 
 			ToolbarBuilder.AddToolBarButton(
 				FTranslationEditorCommands::Get().ExportToPortableObjectFormat, "ExportToPortableObjectFormat", TAttribute<FText>(), TAttribute<FText>(), FSlateIcon(FEditorStyle::GetStyleSetName(), "FontEditor.Export"));
+			ToolbarBuilder.AddToolBarButton(
+				FTranslationEditorCommands::Get().OpenSearchTab, "OpenSearchTab", TAttribute<FText>(), TAttribute<FText>(), FSlateIcon(FEditorStyle::GetStyleSetName(), "BlueprintEditor.FindInBlueprint"));
 		}
 	};
 
@@ -90,6 +93,7 @@ void FTranslationEditorCommands::RegisterCommands()
 	UI_COMMAND( SaveTranslations, "Save", "Saves the translations to file", EUserInterfaceActionType::Button, FInputGesture() );
 	UI_COMMAND( PreviewAllTranslationsInEditor, "Preview in Editor", "Preview All Translations in the Editor UI", EUserInterfaceActionType::Button, FInputGesture() );
 	UI_COMMAND( ExportToPortableObjectFormat, "Export to .PO", "Export to Portable Object Format", EUserInterfaceActionType::Button, FInputGesture() );
+	UI_COMMAND( OpenSearchTab, "Search", "Search Source and Translation Strings", EUserInterfaceActionType::Button, FInputGesture() );
 }
 
 
diff --git a/Engine/Source/Editor/TranslationEditor/Public/TranslationEditorMenu.h b/Engine/Source/Editor/TranslationEditor/Public/TranslationEditorMenu.h
--- a/Engine/Source/Editor/TranslationEditor/Public/TranslationEditorMenu.h
+++ b/Engine/Source/Editor/TranslationEditor/Public/TranslationEditorMenu.h
@@ -39,6 +39,9 @@ public:
 	/** Export to PortableObject format (.po) */
 	TSharedPtr<FUICommandInfo> ExportToPortableObjectFormat;
 
+	/** Open the tab for searching */
+	TSharedPtr<FUICommandInfo> OpenSearchTab;
+
 	/** Initialize commands */
 	virtual void RegisterCommands() OVERRIDE;
 };
