Translation Editor:
- Fixed a bug where updating Context and History tables based on your selection wouldn't work properly if you arranged the tabs containing translation units (Untranslated, Needs Review etc) in separate tab wells or windows.

diff --git a/Engine/Source/Editor/TranslationEditor/Private/TranslationEditor.cpp b/Engine/Source/Editor/TranslationEditor/Private/TranslationEditor.cpp
--- a/Engine/Source/Editor/TranslationEditor/Private/TranslationEditor.cpp
+++ b/Engine/Source/Editor/TranslationEditor/Private/TranslationEditor.cpp
@@ -232,7 +232,7 @@ TSharedRef<SDockTab> FTranslationEditor::SpawnTab_Untranslated( const FSpawnTabA
 	UntranslatedPropertyTable->SetOrientation( EPropertyTableOrientation::AlignPropertiesInColumns );
 	UntranslatedPropertyTable->SetShowRowHeader( true );
 	UntranslatedPropertyTable->SetShowObjectName( false );
-	UntranslatedPropertyTable->OnSelectionChanged()->AddSP( this, &FTranslationEditor::UpdateTranslationUnitSelection );
+	UntranslatedPropertyTable->OnSelectionChanged()->AddSP( this, &FTranslationEditor::UpdateUntranslatedSelection );
 
 	// we want to customize some columns
 	TArray< TSharedRef<class IPropertyTableCustomColumn>> CustomColumns;
@@ -290,7 +290,7 @@ TSharedRef<SDockTab> FTranslationEditor::SpawnTab_Review( const FSpawnTabArgs& A
 	ReviewPropertyTable->SetOrientation( EPropertyTableOrientation::AlignPropertiesInColumns );
 	ReviewPropertyTable->SetShowRowHeader( true );
 	ReviewPropertyTable->SetShowObjectName( false );
-	ReviewPropertyTable->OnSelectionChanged()->AddSP( this, &FTranslationEditor::UpdateTranslationUnitSelection );
+	ReviewPropertyTable->OnSelectionChanged()->AddSP( this, &FTranslationEditor::UpdateNeedsReviewSelection );
 
 	// we want to customize some columns
 	TArray< TSharedRef< class IPropertyTableCustomColumn > > CustomColumns;
@@ -355,7 +355,7 @@ TSharedRef<SDockTab> FTranslationEditor::SpawnTab_Completed( const FSpawnTabArgs
 	CompletedPropertyTable->SetOrientation( EPropertyTableOrientation::AlignPropertiesInColumns );
 	CompletedPropertyTable->SetShowRowHeader( true );
 	CompletedPropertyTable->SetShowObjectName( false );
-	CompletedPropertyTable->OnSelectionChanged()->AddSP( this, &FTranslationEditor::UpdateTranslationUnitSelection );
+	CompletedPropertyTable->OnSelectionChanged()->AddSP( this, &FTranslationEditor::UpdateCompletedSelection );
 
 	// we want to customize some columns
 	TArray< TSharedRef< class IPropertyTableCustomColumn > > CustomColumns;
@@ -413,7 +413,7 @@ TSharedRef<SDockTab> FTranslationEditor::SpawnTab_Search(const FSpawnTabArgs& Ar
 	SearchPropertyTable->SetOrientation(EPropertyTableOrientation::AlignPropertiesInColumns);
 	SearchPropertyTable->SetShowRowHeader(true);
 	SearchPropertyTable->SetShowObjectName(false);
-	SearchPropertyTable->OnSelectionChanged()->AddSP(this, &FTranslationEditor::UpdateTranslationUnitSelection);
+	SearchPropertyTable->OnSelectionChanged()->AddSP(this, &FTranslationEditor::UpdateSearchSelection);
 
 	// we want to customize some columns
 	TArray< TSharedRef< class IPropertyTableCustomColumn > > CustomColumns;
@@ -789,32 +789,48 @@ bool FTranslationEditor::OpenFontPicker( const FString DefaultFile, FString& Out
 	return bOpened;
 }
 
-void FTranslationEditor::UpdateTranslationUnitSelection()
+void FTranslationEditor::UpdateUntranslatedSelection()
 {
-	TSet<TSharedRef<IPropertyTableRow>> SelectedRows;
-
-	// Different Selection based on which tab is in the foreground
 	TSharedPtr<SDockTab> UntranslatedTabSharedPtr = UntranslatedTab.Pin();
-	TSharedPtr<SDockTab> ReviewTabSharedPtr = ReviewTab.Pin();
-	TSharedPtr<SDockTab> CompletedTabSharedPtr = CompletedTab.Pin();
-	TSharedPtr<SDockTab> SearchTabSharedPtr = SearchTab.Pin();
 	if (UntranslatedTabSharedPtr.IsValid() && UntranslatedTabSharedPtr->IsForeground() && UntranslatedPropertyTable.IsValid())
 	{
-		SelectedRows = UntranslatedPropertyTable->GetSelectedRows();
+		TSet<TSharedRef<IPropertyTableRow>> SelectedRows = UntranslatedPropertyTable->GetSelectedRows();
+		UpdateTranslationUnitSelection(SelectedRows);
 	}
-	else if (ReviewTab.IsValid() && ReviewTabSharedPtr->IsForeground() && ReviewPropertyTable.IsValid())
+}
+
+void FTranslationEditor::UpdateNeedsReviewSelection()
+{
+	TSharedPtr<SDockTab> ReviewTabSharedPtr = ReviewTab.Pin();
+	if (ReviewTabSharedPtr.IsValid() && ReviewTabSharedPtr->IsForeground() && ReviewPropertyTable.IsValid())
 	{
-		SelectedRows = ReviewPropertyTable->GetSelectedRows();
+		TSet<TSharedRef<IPropertyTableRow>> SelectedRows = ReviewPropertyTable->GetSelectedRows();
+		UpdateTranslationUnitSelection(SelectedRows);
 	}
-	else if (CompletedTab.IsValid() && CompletedTabSharedPtr->IsForeground() && CompletedPropertyTable.IsValid())
+}
+
+void FTranslationEditor::UpdateCompletedSelection()
+{
+	TSharedPtr<SDockTab> CompletedTabSharedPtr = CompletedTab.Pin();
+	if (CompletedTabSharedPtr.IsValid() && CompletedTabSharedPtr->IsForeground() && CompletedPropertyTable.IsValid())
 	{
-		SelectedRows = CompletedPropertyTable->GetSelectedRows();
+		TSet<TSharedRef<IPropertyTableRow>> SelectedRows = CompletedPropertyTable->GetSelectedRows();
+		UpdateTranslationUnitSelection(SelectedRows);
 	}
-	else if (SearchTab.IsValid() && SearchTabSharedPtr->IsForeground() && SearchPropertyTable.IsValid())
+}
+
+void FTranslationEditor::UpdateSearchSelection()
+{
+	TSharedPtr<SDockTab> SearchTabSharedPtr = SearchTab.Pin();
+	if (SearchTabSharedPtr.IsValid() && SearchTabSharedPtr->IsForeground() && SearchPropertyTable.IsValid())
 	{
-		SelectedRows = SearchPropertyTable->GetSelectedRows();
+		TSet<TSharedRef<IPropertyTableRow>> SelectedRows = SearchPropertyTable->GetSelectedRows();
+		UpdateTranslationUnitSelection(SelectedRows);
 	}
+}
 
+void FTranslationEditor::UpdateTranslationUnitSelection(TSet<TSharedRef<IPropertyTableRow>>& SelectedRows)
+{
 	// Can only really handle single selection
 	if (SelectedRows.Num() == 1)
 	{
diff --git a/Engine/Source/Editor/TranslationEditor/Private/TranslationEditor.h b/Engine/Source/Editor/TranslationEditor/Private/TranslationEditor.h
--- a/Engine/Source/Editor/TranslationEditor/Private/TranslationEditor.h
+++ b/Engine/Source/Editor/TranslationEditor/Private/TranslationEditor.h
@@ -136,7 +136,19 @@ private:
 	void RefreshUI();
 
 	/** Update content when a new translation unit selection is made */
-	void UpdateTranslationUnitSelection();
+	void UpdateTranslationUnitSelection(TSet<TSharedRef<IPropertyTableRow>>& SelectedRows);
+
+	/** Update content when a new translation unit selection is made in the Untranslated PropertyTable */
+	void UpdateUntranslatedSelection();
+
+	/** Update content when a new translation unit selection is made in the Needs Review PropertyTable */
+	void UpdateNeedsReviewSelection();
+
+	/** Update content when a new translation unit selection is made in the Completed PropertyTable */
+	void UpdateCompletedSelection();
+
+	/** Update content when a new translation unit selection is made in the Search PropertyTable */
+	void UpdateSearchSelection();
 
 	/** Update content when a new context selection is made */
 	void UpdateContextSelection();
