TranslationEditor:
- Added InternationalizationExportSettings UObject to use to serialize to a file for the InternationalizationExportCommandlet
- Fixed "Export to .po" to only export the currently loaded project & language combination, to a file the user can now specify

diff --git a/Engine/Source/Editor/TranslationEditor/Private/InternationalizationExportSettings.cpp b/Engine/Source/Editor/TranslationEditor/Private/InternationalizationExportSettings.cpp
new file mode 100644
--- /dev/null
+++ b/Engine/Source/Editor/TranslationEditor/Private/InternationalizationExportSettings.cpp
@@ -0,0 +1,10 @@
+// Copyright 1998-2014 Epic Games, Inc. All Rights Reserved.
+#include "TranslationEditorPrivatePCH.h"
+
+#include "InternationalizationExportSettings.h"
+
+UInternationalizationExportSettings::UInternationalizationExportSettings(const class FPostConstructInitializeProperties& PCIP)
+	: Super(PCIP)
+{
+
+}
\ No newline at end of file
diff --git a/Engine/Source/Editor/TranslationEditor/Private/InternationalizationExportSettings.h b/Engine/Source/Editor/TranslationEditor/Private/InternationalizationExportSettings.h
new file mode 100644
--- /dev/null
+++ b/Engine/Source/Editor/TranslationEditor/Private/InternationalizationExportSettings.h
@@ -0,0 +1,44 @@
+// Copyright 1998-2014 Epic Games, Inc. All Rights Reserved.
+#pragma once
+
+#include "InternationalizationExportSettings.generated.h"
+
+UCLASS(hidecategories = Object, MinimalAPI, config=InternationalizationExport)
+class UInternationalizationExportSettings : public UObject
+{
+	GENERATED_UCLASS_BODY()
+
+public:
+
+	/** Which cultures should be exported */
+	UPROPERTY(Category = CommonSettings, EditAnywhere, config)
+	TArray<FString> CulturesToGenerate;
+
+	/** The commandlet to run */
+	UPROPERTY(Category = GatherTextStep, EditAnywhere, config)
+	FString CommandletClass;
+
+	/** Source for the localization data */
+	UPROPERTY(Category = GatherTextStep, EditAnywhere, config)
+	FString SourcePath;
+
+	/** Destination for the localization data */
+	UPROPERTY(Category = GatherTextStep, EditAnywhere, config)
+	FString DestinationPath;
+
+	/** Name of the manifest file */
+	UPROPERTY(Category = GatherTextStep, EditAnywhere, config)
+	FString ManifestName;
+
+	/** Name of the archive file */
+	UPROPERTY(Category = GatherTextStep, EditAnywhere, config)
+	FString ArchiveName;
+
+	/** Whether or not to export localization data */
+	UPROPERTY(Category = GatherTextStep, EditAnywhere, config)
+	bool bExportLoc;
+
+	/** Whether or not to import localization data */
+	UPROPERTY(Category = GatherTextStep, EditAnywhere, config)
+	bool bImportLoc;
+};
\ No newline at end of file
diff --git a/Engine/Source/Editor/TranslationEditor/Private/TranslationEditor.cpp b/Engine/Source/Editor/TranslationEditor/Private/TranslationEditor.cpp
--- a/Engine/Source/Editor/TranslationEditor/Private/TranslationEditor.cpp
+++ b/Engine/Source/Editor/TranslationEditor/Private/TranslationEditor.cpp
@@ -14,14 +14,12 @@
 #include "Editor/PropertyEditor/Public/PropertyHandle.h"
 #include "Editor/PropertyEditor/Public/PropertyPath.h"
 #include "CustomFontColumn.h"
-
-#include "TranslationUnit.h"
-
 #include "DesktopPlatformModule.h"
-
 #include "IPropertyTableWidgetHandle.h"
 
+#include "TranslationUnit.h"
 #include "SSearchBox.h"
+#include "InternationalizationExportSettings.h"
 
 DEFINE_LOG_CATEGORY_STATIC(LocalizationExport, Log, All);
 
@@ -200,8 +198,8 @@ FText FTranslationEditor::GetToolkitName() const
 	const bool bDirtyState = EditingObject->GetOutermost()->IsDirty();
 
 	FFormatNamedArguments Args;
-	Args.Add( TEXT("Language"), FText::FromString( TranslationTargetLanguage ) );
-	Args.Add( TEXT("ProjectName"), FText::FromString( ProjectName ) );
+	Args.Add(TEXT("Language"), FText::FromString(FPaths::GetBaseFilename(FPaths::GetPath(ArchiveFilePath))));
+	Args.Add(TEXT("ProjectName"), FText::FromString(FPaths::GetBaseFilename(ManifestFilePath)));
 	Args.Add( TEXT("DirtyState"), bDirtyState ? FText::FromString( TEXT( "*" ) ) : FText::GetEmpty() );
 	Args.Add( TEXT("ToolkitName"), GetBaseToolkitName() );
 	return FText::Format( LOCTEXT("TranslationEditorAppLabel", "{Language}{DirtyState} - {ProjectName} - {ToolkitName}"), Args );
@@ -984,56 +982,72 @@ void FTranslationEditor::PreviewAllTranslationsInEditor_Execute()
 
 void FTranslationEditor::ExportToPortableObjectFormat_Execute()
 {
-	// TODO: Add ability to choose output location
-
-	//const FString PortableObjectFileDescription = LOCTEXT("PortableObjectFileDescription", "Portable Object File").ToString();
-	//const FString PortableObjectFileExtension = TEXT("*.po");
-	//const FString FileTypes = FString::Printf(TEXT("%s (%s)|%s"), *PortableObjectFileDescription, *PortableObjectFileExtension, *PortableObjectFileExtension);
-	TArray<FString> OpenFilenames;
-	//IDesktopPlatform* DesktopPlatform = FDesktopPlatformModule::Get();
+	const FString PortableObjectFileDescription = LOCTEXT("PortableObjectFileDescription", "Portable Object File").ToString();
+	const FString PortableObjectFileExtension = TEXT("*.po");
+	const FString FileTypes = FString::Printf(TEXT("%s (%s)|%s"), *PortableObjectFileDescription, *PortableObjectFileExtension, *PortableObjectFileExtension);
+	const FString DefaultFilename = FPaths::GetBaseFilename(ManifestFilePath) + "-" + FPaths::GetBaseFilename(FPaths::GetPath(ArchiveFilePath)) + ".po";
+	const FString DefaultPath = FPaths::GameSavedDir();
+	TArray<FString> SaveFilenames;
+	IDesktopPlatform* DesktopPlatform = FDesktopPlatformModule::Get();
 	bool bSelected = false;
 	
-	// Prompt the user for the filenames
-	//if (DesktopPlatform)
-	//{
-	//	void* ParentWindowWindowHandle = NULL;
-
-	//	const TSharedPtr<SWindow>& ParentWindow = FSlateApplication::Get().FindWidgetWindow(PreviewTextBlock);
-	//	if (ParentWindow.IsValid() && ParentWindow->GetNativeWindow().IsValid())
-	//	{
-	//		ParentWindowWindowHandle = ParentWindow->GetNativeWindow()->GetOSWindowHandle();
-	//	}
-
-	//	bSelected = DesktopPlatform->SaveFileDialog(
-	//		ParentWindowWindowHandle,
-	//		LOCTEXT("ChooseExportLocationWindowTitle", "Choose Export Location").ToString(),
-	//		FPaths::GameSavedDir(),
-	//		"ExportedLocalization.po",
-	//		FileTypes,
-	//		EFileDialogFlags::None,
-	//		OpenFilenames
-	//		);
-	//}
-
-	FString OutFile = FPaths::EngineSavedDir() / "LocalizationExport" / "ExportedLocalization.po";
-
-	if (bSelected && OpenFilenames.Num() > 0)
+	// Prompt the user for the filename
+	if (DesktopPlatform)
 	{
-		OutFile = OpenFilenames[0];
+		void* ParentWindowWindowHandle = NULL;
+
+		const TSharedPtr<SWindow>& ParentWindow = FSlateApplication::Get().FindWidgetWindow(PreviewTextBlock);
+		if (ParentWindow.IsValid() && ParentWindow->GetNativeWindow().IsValid())
+		{
+			ParentWindowWindowHandle = ParentWindow->GetNativeWindow()->GetOSWindowHandle();
+		}
+
+		bSelected = DesktopPlatform->SaveFileDialog(
+			ParentWindowWindowHandle,
+			LOCTEXT("ChooseExportLocationWindowTitle", "Choose Export Location").ToString(),
+			FPaths::GameSavedDir(),
+			DefaultFilename,
+			FileTypes,
+			EFileDialogFlags::None,
+			SaveFilenames
+			);
 	}
 
-	// For now, just run every section in the config file
+	GWarn->BeginSlowTask(LOCTEXT("ExportingInternationalization", "Exporting Internationalization Data..."), true);
+
+	// Write translation data first to ensure all changes are exported
+	DataManager->WriteTranslationData();
+
+	UInternationalizationExportSettings* ExportSettings = NewObject<UInternationalizationExportSettings>();
+	ExportSettings->CulturesToGenerate.Empty();
+	ExportSettings->CulturesToGenerate.Add(FPaths::GetBaseFilename(FPaths::GetPath(ArchiveFilePath)));
+	ExportSettings->CommandletClass = "InternationalizationExport";
+	ExportSettings->SourcePath = FPaths::GetPath(ManifestFilePath);
+	ExportSettings->ManifestName = FPaths::GetBaseFilename(ManifestFilePath) + ".manifest";
+	ExportSettings->ArchiveName = FPaths::GetBaseFilename(ManifestFilePath) + ".archive";
+	ExportSettings->bExportLoc = true;
+	ExportSettings->bImportLoc = false;
+
+	ExportSettings->DestinationPath = DefaultPath / DefaultFilename;
+
+	if (bSelected && SaveFilenames.Num() > 0)
+	{
+		ExportSettings->DestinationPath = SaveFilenames[0];
+	}
+
+	// Write these settings to a temporary config file that the Internationalization Export Commandlet will read
+	FString TempConfigFilepath = FPaths::GameSavedDir() / "Config" / "InternationalizationExport.ini";
+	ExportSettings->SaveConfig(CPF_Config, *TempConfigFilepath);
+
+	// Using .ini config saving means these settings will be saved in the GetClass()->GetPathName() section
 	TArray<FString> ConfigSections;
-	ConfigSections.Add("GatherTextStep0");
-	ConfigSections.Add("GatherTextStep1");
-	ConfigSections.Add("GatherTextStep2");
-	ConfigSections.Add("GatherTextStep3");
+	ConfigSections.Add(ExportSettings->GetClass()->GetPathName());
 
 	for (FString& ConfigSection : ConfigSections)
 	{
 		// Spawn the LocalizationExport commandlet, and run its log output back into ours
 		FString AppURL = FPlatformProcess::ExecutableName(true);
-		FString Parameters = FString("-run=InternationalizationExport -config=") + FPaths::EngineConfigDir() / "Localization" / "PortableObjectExport.ini -section=" + ConfigSection;
+		FString Parameters = FString("-run=InternationalizationExport -config=") + TempConfigFilepath + " -section=" + ConfigSection;
 
 		void* WritePipe;
 		void* ReadPipe;
@@ -1070,6 +1084,7 @@ void FTranslationEditor::ExportToPortableObjectFormat_Execute()
 		}
 	}
 
+	GWarn->EndSlowTask();
 }
 
 void FTranslationEditor::OnFilterTextChanged(const FText& InFilterText)
diff --git a/Engine/Source/Editor/TranslationEditor/Private/TranslationEditor.h b/Engine/Source/Editor/TranslationEditor/Private/TranslationEditor.h
--- a/Engine/Source/Editor/TranslationEditor/Private/TranslationEditor.h
+++ b/Engine/Source/Editor/TranslationEditor/Private/TranslationEditor.h
@@ -15,9 +15,9 @@ public:
 	/**
 	 *	Creates a new FTranslationEditor and calls Initialize
 	 */
-	static TSharedRef< FTranslationEditor > Create( TSharedRef< FTranslationDataManager > DataManager, const FString& ProjectName, const FString& TranslationTargetLanguage )
+	static TSharedRef< FTranslationEditor > Create(TSharedRef< FTranslationDataManager > DataManager, const FString& InManifestFile, const FString& InArchiveFile)
 	{
-		TSharedRef< FTranslationEditor > TranslationEditor = MakeShareable( new FTranslationEditor( DataManager, ProjectName, TranslationTargetLanguage ) );
+		TSharedRef< FTranslationEditor > TranslationEditor = MakeShareable(new FTranslationEditor(DataManager, InManifestFile, InArchiveFile));
 
 		// Some stuff that needs to use the "this" pointer is done in Initialize (because it can't be done in the constructor)
 		TranslationEditor->Initialize();
@@ -69,8 +69,8 @@ private:
 				.Font(TranslationTargetFont))
 	, NamespaceTextBlock(SNew(STextBlock)
 				.Text(FText::FromString("")))
-	, ProjectName(FPaths::GetBaseFilename(InManifestFile))
-	, TranslationTargetLanguage(FPaths::GetBaseFilename(FPaths::GetPath(InArchiveFile)))
+	, ManifestFilePath(InManifestFile)
+	, ArchiveFilePath(InArchiveFile)
 	{}
 
 	/** Does some things we can't do in the constructor because we can't get a SharedRef to "this" there */ 
@@ -242,7 +242,7 @@ private:
 	TSharedRef<STextBlock> NamespaceTextBlock;
 
 	/** Name of the project we are translating for */
-	FString ProjectName;
+	FString ManifestFilePath;
 	/** Name of the language we are translating to */
-	FString TranslationTargetLanguage;
+	FString ArchiveFilePath;
 };
