Translation Editor:
- Added "Import from .PO" option
- Displays a new tab after import showing you what changed
- Various other fixes related to importing / exporting

diff --git a/Engine/Source/Editor/TranslationEditor/Private/TranslationDataManager.cpp b/Engine/Source/Editor/TranslationEditor/Private/TranslationDataManager.cpp
--- a/Engine/Source/Editor/TranslationEditor/Private/TranslationDataManager.cpp
+++ b/Engine/Source/Editor/TranslationEditor/Private/TranslationDataManager.cpp
@@ -38,10 +38,9 @@ FTranslationDataManager::FTranslationDataManager( const FString& InManifestFileP
 			TranslationEditorMessageLog.Open(EMessageSeverity::Error);
 		}
 
-		ArchivePtr = ReadArchive( ManifestAtHeadRevision);
+		ArchivePtr = ReadArchive();
 		if (ArchivePtr.IsValid())
 		{
-			TSharedRef< FInternationalizationArchive > Archive = ArchivePtr.ToSharedRef();
 			int32 NumManifestEntriesParsed = 0;
 
 			GWarn->BeginSlowTask(LOCTEXT("LoadingCurrentManifest", "Loading Entries from Current Translation Manifest..."), true);
@@ -76,66 +75,7 @@ FTranslationDataManager::FTranslationDataManager( const FString& InManifestFileP
 			// Once we have the context information, we can look up history
 			GetHistoryForTranslationUnits(TranslationUnits, ManifestFilePath);
 
-			GWarn->BeginSlowTask(LOCTEXT("LoadingArchiveEntries", "Loading Entries from Translation Archive..."), true);
-			for (int32 CurrentTranslationUnitIndex = 0; CurrentTranslationUnitIndex < TranslationUnits.Num(); ++CurrentTranslationUnitIndex)
-			{
-				UTranslationUnit* TranslationUnit = TranslationUnits[CurrentTranslationUnitIndex];
-				AllTranslations.Add(TranslationUnit);
-
-				GWarn->StatusUpdate(CurrentTranslationUnitIndex, TranslationUnits.Num(), FText::Format(LOCTEXT("LoadingCurrentArchiveEntries", "Loading Entry {0} of {1} from Translation Archive..."), FText::AsNumber(CurrentTranslationUnitIndex), FText::AsNumber(TranslationUnits.Num())));
-
-				const FLocItem SourceSearch(TranslationUnit->Source);
-				TSharedPtr<FArchiveEntry> ArchiveEntry = Archive->FindEntryBySource(TranslationUnit->Namespace, SourceSearch, NULL);
-				if ( ArchiveEntry.IsValid() )
-				{
-					const FString UnescapedTranslatedString = ArchiveEntry->Translation.Text.ReplaceEscapedCharWithChar();
-
-					if(UnescapedTranslatedString.IsEmpty())
-					{
-						bool bHasTranslationHistory = false;
-						int32 MostRecentNonNullTranslationIndex = -1;
-						int32 ContextForRecentTranslation = -1;
-
-						for (int32 ContextIndex = 0; ContextIndex < TranslationUnit->Contexts.Num(); ++ContextIndex)
-						{
-							for (int32 ChangeIndex = 0; ChangeIndex < TranslationUnit->Contexts[ContextIndex].Changes.Num(); ++ChangeIndex)
-							{
-								if (!(TranslationUnit->Contexts[ContextIndex].Changes[ChangeIndex].Translation.IsEmpty()))
-								{
-									bHasTranslationHistory = true;
-									MostRecentNonNullTranslationIndex = ChangeIndex;
-									ContextForRecentTranslation = ContextIndex;
-									break;
-								}
-							}
-					
-							if (bHasTranslationHistory)
-							{
-								break;
-							}
-						}
-
-						// If we have history, but current translation is empty, this goes in the Needs Review tab
-						if (bHasTranslationHistory)
-						{
-							// Offer the most recent translation (for the first context in the list) as a suggestion or starting point (not saved unless user checks "Has Been Reviewed")
-							TranslationUnit->Translation = TranslationUnit->Contexts[ContextForRecentTranslation].Changes[MostRecentNonNullTranslationIndex].Translation;
-							Review.Add(TranslationUnit);
-						}
-						else
-						{
-							Untranslated.Add(TranslationUnit);
-						}
-					}
-					else
-					{
-						TranslationUnit->Translation = UnescapedTranslatedString;
-						TranslationUnit->HasBeenReviewed = true;
-						Complete.Add(TranslationUnit);
-					}
-				}
-			}
-			GWarn->EndSlowTask();
+			LoadFromArchive(TranslationUnits);
 		}
 		else  // ArchivePtr.IsValid() is false
 		{
@@ -178,7 +118,7 @@ TSharedPtr< FInternationalizationManifest > FTranslationDataManager::ReadManifes
 	return InternationalizationManifest;
 }
 
-TSharedPtr< FInternationalizationArchive > FTranslationDataManager::ReadArchive(TSharedRef< FInternationalizationManifest >& InternationalizationManifest)
+TSharedPtr< FInternationalizationArchive > FTranslationDataManager::ReadArchive()
 {
 	// Read in any existing archive for this culture.
 	TSharedPtr< FJsonObject > ArchiveJsonObject = ReadJSONTextFile( ArchiveFilePath );
@@ -220,7 +160,7 @@ TSharedPtr<FJsonObject> FTranslationDataManager::ReadJSONTextFile(const FString&
 	return JSONObject;
 }
 
-void FTranslationDataManager::WriteTranslationData()
+bool FTranslationDataManager::WriteTranslationData(bool bForceWrite /*= false*/)
 {
 	check (ArchivePtr.IsValid());
 	TSharedRef< FInternationalizationArchive > Archive = ArchivePtr.ToSharedRef();
@@ -263,13 +203,17 @@ void FTranslationDataManager::WriteTranslationData()
 		}
 	}
 
-	if (bNeedsWrite)
+	bool bSuccess = true;
+
+	if (bForceWrite || bNeedsWrite)
 	{
 		TSharedRef<FJsonObject> FinalArchiveJsonObj = MakeShareable( new FJsonObject );
 		ArchiveSerializer.SerializeArchive( Archive, FinalArchiveJsonObj );
 
-		WriteJSONToTextFile( FinalArchiveJsonObj, ArchiveFilePath );
+		bSuccess = WriteJSONToTextFile( FinalArchiveJsonObj, ArchiveFilePath );
 	}
+
+	return bSuccess;
 }
 
 bool FTranslationDataManager::WriteJSONToTextFile(TSharedRef<FJsonObject>& Output, const FString& Filename)
@@ -560,4 +504,109 @@ void FTranslationDataManager::PopulateSearchResultsUsingFilter(const FString& Se
 	}
 }
 
+void FTranslationDataManager::LoadFromArchive(TArray<UTranslationUnit*>& InTranslationUnits, bool bTrackChanges /*= false*/, bool bReloadFromFile /*=false*/)
+{
+	GWarn->BeginSlowTask(LOCTEXT("LoadingArchiveEntries", "Loading Entries from Translation Archive..."), true);
+
+	if (bReloadFromFile)
+	{
+		ArchivePtr = ReadArchive();
+	}
+
+	if (ArchivePtr.IsValid())
+	{
+		TSharedRef< FInternationalizationArchive > Archive = ArchivePtr.ToSharedRef();
+
+		// Make a local copy of this array before we empty the arrays below (we might have been passed AllTranslations array)
+		TArray<UTranslationUnit*> TranslationUnits;
+		TranslationUnits.Append(InTranslationUnits);
+
+		AllTranslations.Empty();
+		Untranslated.Empty();
+		Review.Empty();
+		Complete.Empty();
+		ChangedOnImport.Empty();
+
+
+		for (int32 CurrentTranslationUnitIndex = 0; CurrentTranslationUnitIndex < TranslationUnits.Num(); ++CurrentTranslationUnitIndex)
+		{
+			UTranslationUnit* TranslationUnit = TranslationUnits[CurrentTranslationUnitIndex];
+			AllTranslations.Add(TranslationUnit);
+
+			GWarn->StatusUpdate(CurrentTranslationUnitIndex, TranslationUnits.Num(), FText::Format(LOCTEXT("LoadingCurrentArchiveEntries", "Loading Entry {0} of {1} from Translation Archive..."), FText::AsNumber(CurrentTranslationUnitIndex), FText::AsNumber(TranslationUnits.Num())));
+
+			const FLocItem SourceSearch(TranslationUnit->Source);
+			TSharedPtr<FArchiveEntry> ArchiveEntry = Archive->FindEntryBySource(TranslationUnit->Namespace, SourceSearch, NULL);
+			if (ArchiveEntry.IsValid())
+			{
+				const FString PreviousTranslation = TranslationUnit->Translation;
+				TranslationUnit->Translation = ""; // Reset to null string
+				const FString UnescapedTranslatedString = ArchiveEntry->Translation.Text.ReplaceEscapedCharWithChar();
+
+				if (UnescapedTranslatedString.IsEmpty())
+				{
+					bool bHasTranslationHistory = false;
+					int32 MostRecentNonNullTranslationIndex = -1;
+					int32 ContextForRecentTranslation = -1;
+
+					for (int32 ContextIndex = 0; ContextIndex < TranslationUnit->Contexts.Num(); ++ContextIndex)
+					{
+						for (int32 ChangeIndex = 0; ChangeIndex < TranslationUnit->Contexts[ContextIndex].Changes.Num(); ++ChangeIndex)
+						{
+							if (!(TranslationUnit->Contexts[ContextIndex].Changes[ChangeIndex].Translation.IsEmpty()))
+							{
+								bHasTranslationHistory = true;
+								MostRecentNonNullTranslationIndex = ChangeIndex;
+								ContextForRecentTranslation = ContextIndex;
+								break;
+							}
+						}
+
+						if (bHasTranslationHistory)
+						{
+							break;
+						}
+					}
+
+					// If we have history, but current translation is empty, this goes in the Needs Review tab
+					if (bHasTranslationHistory)
+					{
+						// Offer the most recent translation (for the first context in the list) as a suggestion or starting point (not saved unless user checks "Has Been Reviewed")
+						TranslationUnit->Translation = TranslationUnit->Contexts[ContextForRecentTranslation].Changes[MostRecentNonNullTranslationIndex].Translation;
+						Review.Add(TranslationUnit);
+					}
+					else
+					{
+						Untranslated.Add(TranslationUnit);
+					}
+				}
+				else
+				{
+					TranslationUnit->Translation = UnescapedTranslatedString;
+					TranslationUnit->HasBeenReviewed = true;
+					Complete.Add(TranslationUnit);
+				}
+
+				// Add to changed array if we're tracking changes
+				if (bTrackChanges && PreviousTranslation != TranslationUnit->Translation)
+				{
+					ChangedOnImport.Add(TranslationUnit);
+					TranslationUnit->TranslationBeforeImport = PreviousTranslation;
+				}
+			}
+		}
+	}
+	else  // ArchivePtr.IsValid() is false
+	{
+		FFormatNamedArguments Arguments;
+		Arguments.Add(TEXT("ArchiveFilePath"), FText::FromString(ArchiveFilePath));
+		FMessageLog TranslationEditorMessageLog("TranslationEditor");
+		TranslationEditorMessageLog.Error(FText::Format(LOCTEXT("FailedToLoadCurrentArchive", "Failed to load most current translation archive ({ArchiveFilePath}), unable to load translations."), Arguments));
+		TranslationEditorMessageLog.Notify(LOCTEXT("TranslationLoadError", "Error Loading Translations!"));
+		TranslationEditorMessageLog.Open(EMessageSeverity::Error);
+	}
+	
+	GWarn->EndSlowTask();
+}
+
 #undef LOCTEXT_NAMESPACE
diff --git a/Engine/Source/Editor/TranslationEditor/Private/TranslationDataManager.h b/Engine/Source/Editor/TranslationEditor/Private/TranslationDataManager.h
--- a/Engine/Source/Editor/TranslationEditor/Private/TranslationDataManager.h
+++ b/Engine/Source/Editor/TranslationEditor/Private/TranslationDataManager.h
@@ -41,9 +41,17 @@ public:
 	{
 		return SearchResults;
 	}
+
+	TArray<UTranslationUnit*>& GetChangedOnImportArray()
+	{
+		return ChangedOnImport;
+	}
 	
-	/** Write the translation data in memory out to .archive file (check out the .archive file first if necessary) */
-	void WriteTranslationData();
+	/** Write the translation data in memory out to .archive file (check out the .archive file first if necessary) 
+	* @param bForceWrite Whether or not to force a write even if the data hasn't changed
+	* @param return Whether or not the write succeeded (true is succeeded)
+	*/
+	bool WriteTranslationData(bool bForceWrite = false);
 
 	/** Delegate called when a TranslationDataObject property is changed */
 	void HandlePropertyChanged(FName PropertyName);
@@ -54,14 +62,17 @@ public:
 	/** Put items in the Search Array if they match this filter */
 	void PopulateSearchResultsUsingFilter(const FString& SearchFilter);
 
+	/** Load (or reload) Translations from Archive file */ 
+	void LoadFromArchive(TArray<UTranslationUnit*>& InTranslationUnits, bool bTrackChanges = false, bool bReloadFromFile = false);
+
 private:
 	/** Read text file into a JSON file */
 	static TSharedPtr<FJsonObject> ReadJSONTextFile( const FString& InFilePath ) ;
 
 	/** Take a path and a manifest name and return a manifest data structure */
 	TSharedPtr< FInternationalizationManifest > ReadManifest ( const FString& ManifestFilePath );
-	/** Take a path and an archive name and return an archive data structure */
-	TSharedPtr< FInternationalizationArchive > ReadArchive( TSharedRef< FInternationalizationManifest >& InternationalizationManifest );
+	/** Retrieve an archive data structure from ArchiveFilePath */
+	TSharedPtr< FInternationalizationArchive > ReadArchive();
 
 	/** Write JSON file to text file */
 	bool WriteJSONToTextFile( TSharedRef<FJsonObject>& Output, const FString& Filename );
@@ -75,6 +86,7 @@ private:
 	TArray<UTranslationUnit*> Review;
 	TArray<UTranslationUnit*> Complete;
 	TArray<UTranslationUnit*> SearchResults;
+	TArray<UTranslationUnit*> ChangedOnImport;
 
 	/** Serializes and deserializes our Archive */
 	FInternationalizationArchiveJsonSerializer ArchiveSerializer;
diff --git a/Engine/Source/Editor/TranslationEditor/Private/TranslationEditor.cpp b/Engine/Source/Editor/TranslationEditor/Private/TranslationEditor.cpp
--- a/Engine/Source/Editor/TranslationEditor/Private/TranslationEditor.cpp
+++ b/Engine/Source/Editor/TranslationEditor/Private/TranslationEditor.cpp
@@ -5,6 +5,7 @@
 #include "TranslationEditor.h"
 #include "Toolkits/IToolkitHost.h"
 #include "WorkspaceMenuStructureModule.h"
+#include "MessageLog.h"
 
 #include "Editor/PropertyEditor/Public/PropertyEditorModule.h"
 #include "Editor/PropertyEditor/Public/IPropertyTable.h"
@@ -32,6 +33,7 @@ const FName FTranslationEditor::PreviewTabId( TEXT( "TranslationEditor_Preview"
 const FName FTranslationEditor::ContextTabId( TEXT( "TranslationEditor_Context" ) );
 const FName FTranslationEditor::HistoryTabId( TEXT( "TranslationEditor_History" ) );
 const FName FTranslationEditor::SearchTabId( TEXT( "TranslationEditor_Search" ) );
+const FName FTranslationEditor::ChangedOnImportTabId( TEXT( "TranslationEditor_ChangedOnImport" ) );
 
 void FTranslationEditor::Initialize()
 {
@@ -72,9 +74,13 @@ void FTranslationEditor::RegisterTabSpawners(const TSharedRef<class FTabManager>
 		.SetDisplayName( LOCTEXT("HistoryTab", "History") )
 		.SetGroup( MenuStructure.GetAssetEditorCategory() );
 
-	TabManager->RegisterTabSpawner(SearchTabId, FOnSpawnTab::CreateSP(this, &FTranslationEditor::SpawnTab_Search))
+	TabManager->RegisterTabSpawner( SearchTabId, FOnSpawnTab::CreateSP(this, &FTranslationEditor::SpawnTab_Search) )
 		.SetDisplayName(LOCTEXT("SearchTab", "Search"))
 		.SetGroup(MenuStructure.GetAssetEditorCategory());
+
+	TabManager->RegisterTabSpawner( ChangedOnImportTabId, FOnSpawnTab::CreateSP(this, &FTranslationEditor::SpawnTab_ChangedOnImport) )
+		.SetDisplayName(LOCTEXT("ChangedOnImportTab", "ChangedOnImport"))
+		.SetGroup(MenuStructure.GetAssetEditorCategory());
 }
 
 void FTranslationEditor::UnregisterTabSpawners(const TSharedRef<class FTabManager>& TabManager)
@@ -86,6 +92,7 @@ void FTranslationEditor::UnregisterTabSpawners(const TSharedRef<class FTabManage
 	TabManager->UnregisterTabSpawner( ContextTabId );
 	TabManager->UnregisterTabSpawner(HistoryTabId);
 	TabManager->UnregisterTabSpawner(SearchTabId);
+	TabManager->UnregisterTabSpawner(ChangedOnImportTabId);
 }
 
 void FTranslationEditor::InitTranslationEditor( const EToolkitMode::Type Mode, const TSharedPtr< class IToolkitHost >& InitToolkitHost )
@@ -111,6 +118,7 @@ void FTranslationEditor::InitTranslationEditor( const EToolkitMode::Type Mode, c
 			->AddTab( ReviewTabId,  ETabState::OpenedTab )
 			->AddTab( CompletedTabId,  ETabState::OpenedTab )
 			->AddTab( SearchTabId, ETabState::ClosedTab )
+			->AddTab( ChangedOnImportTabId, ETabState::ClosedTab )
 		)
 		->Split
 		(
@@ -474,6 +482,66 @@ TSharedRef<SDockTab> FTranslationEditor::SpawnTab_Search(const FSpawnTabArgs& Ar
 	return NewDockTab;
 }
 
+TSharedRef<SDockTab> FTranslationEditor::SpawnTab_ChangedOnImport(const FSpawnTabArgs& Args)
+{
+	check(Args.GetTabId().TabType == ChangedOnImportTabId);
+
+	UProperty* SourceProperty = FindField<UProperty>(UTranslationUnit::StaticClass(), "Source");
+	UProperty* TranslationBeforeImportProperty = FindField<UProperty>(UTranslationUnit::StaticClass(), "TranslationBeforeImport");
+	UProperty* TranslationProperty = FindField<UProperty>(UTranslationUnit::StaticClass(), "Translation");
+
+	FPropertyEditorModule& PropertyEditorModule = FModuleManager::GetModuleChecked<FPropertyEditorModule>("PropertyEditor");
+
+	// create empty property table
+	ChangedOnImportPropertyTable = PropertyEditorModule.CreatePropertyTable();
+	ChangedOnImportPropertyTable->SetIsUserAllowedToChangeRoot(false);
+	ChangedOnImportPropertyTable->SetOrientation(EPropertyTableOrientation::AlignPropertiesInColumns);
+	ChangedOnImportPropertyTable->SetShowRowHeader(true);
+	ChangedOnImportPropertyTable->SetShowObjectName(false);
+	ChangedOnImportPropertyTable->OnSelectionChanged()->AddSP(this, &FTranslationEditor::UpdateSearchSelection);
+
+	// we want to customize some columns
+	TArray< TSharedRef< class IPropertyTableCustomColumn > > CustomColumns;
+	SourceColumn->AddSupportedProperty(SourceProperty);
+	TranslationColumn->AddSupportedProperty(TranslationProperty);
+	CustomColumns.Add(SourceColumn);
+	CustomColumns.Add(TranslationColumn);
+
+	ChangedOnImportPropertyTable->SetObjects((TArray<UObject*>&)DataManager->GetSearchResultsArray());
+
+	// Add the columns we want to display
+	ChangedOnImportPropertyTable->AddColumn((TWeakObjectPtr<UProperty>)FindField<UProperty>(UTranslationUnit::StaticClass(), "Source"));
+	ChangedOnImportPropertyTable->AddColumn((TWeakObjectPtr<UProperty>)FindField<UProperty>(UTranslationUnit::StaticClass(), "TranslationBeforeImport"));
+	ChangedOnImportPropertyTable->AddColumn((TWeakObjectPtr<UProperty>)FindField<UProperty>(UTranslationUnit::StaticClass(), "Translation"));
+
+	// Freeze columns, don't want user to remove them
+	TArray<TSharedRef<IPropertyTableColumn>> Columns = ChangedOnImportPropertyTable->GetColumns();
+	for (TSharedRef<IPropertyTableColumn> Column : Columns)
+	{
+		Column->SetFrozen(true);
+	}
+
+	SearchPropertyTableWidgetHandle = PropertyEditorModule.CreatePropertyTableWidgetHandle(ChangedOnImportPropertyTable.ToSharedRef(), CustomColumns);
+	TSharedRef<SWidget> PropertyTableWidget = SearchPropertyTableWidgetHandle->GetWidget();
+
+	TSharedRef<SDockTab> NewDockTab = SNew(SDockTab)
+		.Icon(FEditorStyle::GetBrush("TranslationEditor.Tabs.Properties"))
+		.Label(LOCTEXT("ChangedOnImportTabTitle", "Changed on Import"))
+		.TabColorScale(GetTabColorScale())
+		[
+			SNew(SBorder)
+			.BorderImage(FEditorStyle::GetBrush("ToolPanel.GroupBorder"))
+			.Padding(0.0f)
+			[
+				PropertyTableWidget
+			]
+		];
+
+	ChangedOnImportTab = NewDockTab;
+
+	return NewDockTab;
+}
+
 TSharedRef<SDockTab> FTranslationEditor::SpawnTab_Preview( const FSpawnTabArgs& Args )
 {
 	check( Args.GetTabId().TabType == PreviewTabId );
@@ -675,6 +743,10 @@ void FTranslationEditor::MapActions()
 		FExecuteAction::CreateSP(this, &FTranslationEditor::ExportToPortableObjectFormat_Execute),
 		FCanExecuteAction());
 
+	ToolkitCommands->MapAction(FTranslationEditorCommands::Get().ImportFromPortableObjectFormat,
+		FExecuteAction::CreateSP(this, &FTranslationEditor::ImportFromPortableObjectFormat_Execute),
+		FCanExecuteAction());
+
 	ToolkitCommands->MapAction(FTranslationEditorCommands::Get().OpenSearchTab,
 		FExecuteAction::CreateSP(this, &FTranslationEditor::OpenSearchTab_Execute),
 		FCanExecuteAction());
@@ -742,6 +814,10 @@ void FTranslationEditor::RefreshUI()
 	{
 		SearchPropertyTableWidgetHandle->RequestRefresh();
 	}
+	if (ChangedOnImportPropertyTableWidgetHandle.IsValid())
+	{
+		ChangedOnImportPropertyTableWidgetHandle->RequestRefresh();
+	}
 }
 
 bool FTranslationEditor::OpenFontPicker( const FString DefaultFile, FString& OutFile )
@@ -827,6 +903,16 @@ void FTranslationEditor::UpdateSearchSelection()
 	}
 }
 
+void FTranslationEditor::UpdateChangedOnImportSelection()
+{
+	TSharedPtr<SDockTab> ChangedOnImportTabSharedPtr = SearchTab.Pin();
+	if (ChangedOnImportTabSharedPtr.IsValid() && ChangedOnImportTabSharedPtr->IsForeground() && ChangedOnImportPropertyTable.IsValid())
+	{
+		TSet<TSharedRef<IPropertyTableRow>> SelectedRows = ChangedOnImportPropertyTable->GetSelectedRows();
+		UpdateTranslationUnitSelection(SelectedRows);
+	}
+}
+
 void FTranslationEditor::UpdateTranslationUnitSelection(TSet<TSharedRef<IPropertyTableRow>>& SelectedRows)
 {
 	// Can only really handle single selection
@@ -986,10 +1072,15 @@ void FTranslationEditor::ExportToPortableObjectFormat_Execute()
 	const FString PortableObjectFileExtension = TEXT("*.po");
 	const FString FileTypes = FString::Printf(TEXT("%s (%s)|%s"), *PortableObjectFileDescription, *PortableObjectFileExtension, *PortableObjectFileExtension);
 	const FString DefaultFilename = FPaths::GetBaseFilename(ManifestFilePath) + "-" + FPaths::GetBaseFilename(FPaths::GetPath(ArchiveFilePath)) + ".po";
-	const FString DefaultPath = FPaths::GameSavedDir();
+	FString DefaultPath = FPaths::GameSavedDir();
+	if (LastExportFilePath != "")
+	{
+		DefaultPath = LastExportFilePath;
+	}
 	TArray<FString> SaveFilenames;
 	IDesktopPlatform* DesktopPlatform = FDesktopPlatformModule::Get();
 	bool bSelected = false;
+	bool bHadError = false;
 	
 	// Prompt the user for the filename
 	if (DesktopPlatform)
@@ -1005,7 +1096,7 @@ void FTranslationEditor::ExportToPortableObjectFormat_Execute()
 		bSelected = DesktopPlatform->SaveFileDialog(
 			ParentWindowWindowHandle,
 			LOCTEXT("ChooseExportLocationWindowTitle", "Choose Export Location").ToString(),
-			FPaths::GameSavedDir(),
+			LastExportFilePath,
 			DefaultFilename,
 			FileTypes,
 			EFileDialogFlags::None,
@@ -1013,78 +1104,238 @@ void FTranslationEditor::ExportToPortableObjectFormat_Execute()
 			);
 	}
 
-	GWarn->BeginSlowTask(LOCTEXT("ExportingInternationalization", "Exporting Internationalization Data..."), true);
+	if (bSelected)
+	{
+		GWarn->BeginSlowTask(LOCTEXT("ExportingInternationalization", "Exporting Internationalization Data..."), true);
 
-	// Write translation data first to ensure all changes are exported
-	DataManager->WriteTranslationData();
+		// Write translation data first to ensure all changes are exported
+		DataManager->WriteTranslationData();
 
-	UInternationalizationExportSettings* ExportSettings = NewObject<UInternationalizationExportSettings>();
-	ExportSettings->CulturesToGenerate.Empty();
-	ExportSettings->CulturesToGenerate.Add(FPaths::GetBaseFilename(FPaths::GetPath(ArchiveFilePath)));
-	ExportSettings->CommandletClass = "InternationalizationExport";
-	ExportSettings->SourcePath = FPaths::GetPath(ManifestFilePath);
-	ExportSettings->ManifestName = FPaths::GetBaseFilename(ManifestFilePath) + ".manifest";
-	ExportSettings->ArchiveName = FPaths::GetBaseFilename(ManifestFilePath) + ".archive";
-	ExportSettings->bExportLoc = true;
-	ExportSettings->bImportLoc = false;
+		UInternationalizationExportSettings* ExportSettings = NewObject<UInternationalizationExportSettings>();
+		ExportSettings->CulturesToGenerate.Empty();
+		ExportSettings->CulturesToGenerate.Add(FPaths::GetBaseFilename(FPaths::GetPath(ArchiveFilePath)));
+		ExportSettings->CommandletClass = "InternationalizationExport";
+		ExportSettings->SourcePath = FPaths::GetPath(ManifestFilePath);
+		ExportSettings->ManifestName = FPaths::GetBaseFilename(ManifestFilePath) + ".manifest";
+		ExportSettings->ArchiveName = FPaths::GetBaseFilename(ManifestFilePath) + ".archive";
+		ExportSettings->bExportLoc = true;
+		ExportSettings->bImportLoc = false;
 
-	ExportSettings->DestinationPath = DefaultPath / DefaultFilename;
+		ExportSettings->DestinationPath = DefaultPath / DefaultFilename;
 
-	if (bSelected && SaveFilenames.Num() > 0)
-	{
-		ExportSettings->DestinationPath = SaveFilenames[0];
-	}
+		if (SaveFilenames.Num() > 0)
+		{
+			ExportSettings->DestinationPath = SaveFilenames[0];
+			LastExportFilePath = FPaths::GetPath(SaveFilenames[0]);
+		}
 
-	// Write these settings to a temporary config file that the Internationalization Export Commandlet will read
-	FString TempConfigFilepath = FPaths::GameSavedDir() / "Config" / "InternationalizationExport.ini";
-	ExportSettings->SaveConfig(CPF_Config, *TempConfigFilepath);
+		// Write these settings to a temporary config file that the Internationalization Export Commandlet will read
+		FString TempConfigFilepath = FPaths::GameSavedDir() / "Config" / "InternationalizationExport.ini";
+		ExportSettings->SaveConfig(CPF_Config, *TempConfigFilepath);
 
-	// Using .ini config saving means these settings will be saved in the GetClass()->GetPathName() section
-	TArray<FString> ConfigSections;
-	ConfigSections.Add(ExportSettings->GetClass()->GetPathName());
+		// Using .ini config saving means these settings will be saved in the GetClass()->GetPathName() section
+		TArray<FString> ConfigSections;
+		ConfigSections.Add(ExportSettings->GetClass()->GetPathName());
+		FMessageLog TranslationEditorMessageLog("TranslationEditor");
 
-	for (FString& ConfigSection : ConfigSections)
-	{
-		// Spawn the LocalizationExport commandlet, and run its log output back into ours
-		FString AppURL = FPlatformProcess::ExecutableName(true);
-		FString Parameters = FString("-run=InternationalizationExport -config=") + TempConfigFilepath + " -section=" + ConfigSection;
+		for (FString& ConfigSection : ConfigSections)
+		{
+			// Spawn the LocalizationExport commandlet, and run its log output back into ours
+			FString AppURL = FPlatformProcess::ExecutableName(true);
+			FString Parameters = FString("-run=InternationalizationExport -config=") + TempConfigFilepath + " -section=" + ConfigSection;
 
-		void* WritePipe;
-		void* ReadPipe;
-		FPlatformProcess::CreatePipe(ReadPipe, WritePipe);
-		FProcHandle ProcessHandle = FPlatformProcess::CreateProc(*AppURL, *Parameters, false, false, false, NULL, 0, NULL, WritePipe);
+			void* WritePipe;
+			void* ReadPipe;
+			FPlatformProcess::CreatePipe(ReadPipe, WritePipe);
+			FProcHandle ProcessHandle = FPlatformProcess::CreateProc(*AppURL, *Parameters, false, false, false, NULL, 0, NULL, WritePipe);
 
-		while (FPlatformProcess::IsProcRunning(ProcessHandle))
-		{
+			while (FPlatformProcess::IsProcRunning(ProcessHandle))
+			{
+				FString NewLine = FPlatformProcess::ReadPipe(ReadPipe);
+				if (NewLine.Len() > 0)
+				{
+					UE_LOG(LocalizationExport, Log, TEXT("%s"), *NewLine);
+					FFormatNamedArguments Arguments;
+					Arguments.Add(TEXT("LogMessage"), FText::FromString(NewLine));
+					TranslationEditorMessageLog.Info(FText::Format(LOCTEXT("LocalizationExportLog", "Localization Export Log: {LogMessage}"), Arguments));
+				}
+
+				FPlatformProcess::Sleep(0.25);
+			}
 			FString NewLine = FPlatformProcess::ReadPipe(ReadPipe);
 			if (NewLine.Len() > 0)
 			{
 				UE_LOG(LocalizationExport, Log, TEXT("%s"), *NewLine);
+				FFormatNamedArguments Arguments;
+				Arguments.Add(TEXT("LogMessage"), FText::FromString(NewLine));
+				TranslationEditorMessageLog.Info(FText::Format(LOCTEXT("LocalizationExportLog", "Localization Export Log: {LogMessage}"), Arguments));
 			}
 
 			FPlatformProcess::Sleep(0.25);
+			FPlatformProcess::ClosePipe(ReadPipe, WritePipe);
+
+			int32 ReturnCode;
+			if (!FPlatformProcess::GetProcReturnCode(ProcessHandle, &ReturnCode))
+			{
+				bHadError = true;
+			}
+			else if (ReturnCode != 0)
+			{
+				bHadError = true;
+			}
 		}
-		FString NewLine = FPlatformProcess::ReadPipe(ReadPipe);
-		if (NewLine.Len() > 0)
+
+		GWarn->EndSlowTask();
+
+		if (bHadError)
 		{
-			UE_LOG(LocalizationExport, Log, TEXT("%s"), *NewLine);
+			TranslationEditorMessageLog.Error(LOCTEXT("FailedToExportLocalization", "Failed to export localization!"));
+			TranslationEditorMessageLog.Notify(LOCTEXT("FailedToExportLocalization", "Failed to export localization!"));
+			TranslationEditorMessageLog.Open(EMessageSeverity::Error);
 		}
+	}
+}
 
-		FPlatformProcess::Sleep(0.25);
-		FPlatformProcess::ClosePipe(ReadPipe, WritePipe);
+void FTranslationEditor::ImportFromPortableObjectFormat_Execute()
+{
+	const FString PortableObjectFileDescription = LOCTEXT("PortableObjectFileDescription", "Portable Object File").ToString();
+	const FString PortableObjectFileExtension = TEXT("*.po");
+	const FString FileTypes = FString::Printf(TEXT("%s (%s)|%s"), *PortableObjectFileDescription, *PortableObjectFileExtension, *PortableObjectFileExtension);
+	FString DefaultPath = FPaths::GameSavedDir();
+	if (LastImportFilePath != "")
+	{
+		DefaultPath = LastImportFilePath;
+	}
+	TArray<FString> OpenFilenames;
+	bool bHadError = false;
+	IDesktopPlatform* DesktopPlatform = FDesktopPlatformModule::Get();
 
-		int32 ReturnCode;
-		if (!FPlatformProcess::GetProcReturnCode(ProcessHandle, &ReturnCode))
+	bool bOpened = false;
+	if (DesktopPlatform)
+	{
+		void* ParentWindowWindowHandle = NULL;
+
+		const TSharedPtr<SWindow>& ParentWindow = FSlateApplication::Get().FindWidgetWindow(PreviewTextBlock);
+		if (ParentWindow.IsValid() && ParentWindow->GetNativeWindow().IsValid())
 		{
-			// TODO: print some error
+			ParentWindowWindowHandle = ParentWindow->GetNativeWindow()->GetOSWindowHandle();
 		}
-		else if (ReturnCode != 0)
+
+		bOpened = DesktopPlatform->OpenFileDialog(
+			ParentWindowWindowHandle,
+			LOCTEXT("ChooseImportLocationWindowTitle", "Choose File to Import").ToString(),
+			DefaultPath,
+			TEXT(""),
+			FileTypes,
+			EFileDialogFlags::None,
+			OpenFilenames
+			);
+	}
+
+	if (bOpened)
+	{
+		UInternationalizationExportSettings* ImportSettings = NewObject<UInternationalizationExportSettings>();
+		ImportSettings->CulturesToGenerate.Empty();
+		ImportSettings->CulturesToGenerate.Add(FPaths::GetBaseFilename(FPaths::GetPath(ArchiveFilePath)));
+		ImportSettings->CommandletClass = "InternationalizationExport";
+		ImportSettings->DestinationPath = FPaths::GetPath(ManifestFilePath);
+		ImportSettings->ManifestName = FPaths::GetBaseFilename(ManifestFilePath) + ".manifest";
+		ImportSettings->ArchiveName = FPaths::GetBaseFilename(ManifestFilePath) + ".archive";
+		ImportSettings->bExportLoc = false;
+		ImportSettings->bImportLoc = true;
+
+		ImportSettings->SourcePath = DefaultPath / FPaths::GetBaseFilename(ManifestFilePath);
+
+		if (OpenFilenames.Num() > 0)
 		{
-			// TODO: print some error
+			ImportSettings->SourcePath = OpenFilenames[0];
+			LastImportFilePath = FPaths::GetPath(OpenFilenames[0]);
 		}
-	}
 
-	GWarn->EndSlowTask();
+		// Write translation data first to ensure all changes are exported
+		bHadError = !(DataManager->WriteTranslationData(true));
+
+		if (!bHadError)
+		{
+			GWarn->BeginSlowTask(LOCTEXT("ImportingInternationalization", "Importing Internationalization Data..."), true);
+
+			// Write these settings to a temporary config file that the Internationalization Export Commandlet will read
+			FString TempConfigFilepath = FPaths::GameSavedDir() / "Config" / "InternationalizationExport.ini";
+			ImportSettings->SaveConfig(CPF_Config, *TempConfigFilepath);
+
+			// Using .ini config saving means these settings will be saved in the GetClass()->GetPathName() section
+			TArray<FString> ConfigSections;
+			ConfigSections.Add(ImportSettings->GetClass()->GetPathName());
+			FMessageLog TranslationEditorMessageLog("TranslationEditor");
+
+			for (FString& ConfigSection : ConfigSections)
+			{
+				// Spawn the LocalizationExport commandlet, and run its log output back into ours
+				FString AppURL = FPlatformProcess::ExecutableName(true);
+				FString Parameters = FString("-run=InternationalizationExport -config=") + TempConfigFilepath + " -section=" + ConfigSection;
+
+				void* WritePipe;
+				void* ReadPipe;
+				FPlatformProcess::CreatePipe(ReadPipe, WritePipe);
+				FProcHandle ProcessHandle = FPlatformProcess::CreateProc(*AppURL, *Parameters, false, false, false, NULL, 0, NULL, WritePipe);
+
+				while (FPlatformProcess::IsProcRunning(ProcessHandle))
+				{
+					FString NewLine = FPlatformProcess::ReadPipe(ReadPipe);
+					if (NewLine.Len() > 0)
+					{
+						UE_LOG(LocalizationExport, Log, TEXT("%s"), *NewLine);
+						FFormatNamedArguments Arguments;
+						Arguments.Add(TEXT("LogMessage"), FText::FromString(NewLine));
+						TranslationEditorMessageLog.Info(FText::Format(LOCTEXT("LocalizationImportLog", "Localization Import Log: {LogMessage}"), Arguments));
+					}
+
+					FPlatformProcess::Sleep(0.25);
+				}
+				FString NewLine = FPlatformProcess::ReadPipe(ReadPipe);
+				if (NewLine.Len() > 0)
+				{
+					UE_LOG(LocalizationExport, Log, TEXT("%s"), *NewLine);
+					FFormatNamedArguments Arguments;
+					Arguments.Add(TEXT("LogMessage"), FText::FromString(NewLine));
+					TranslationEditorMessageLog.Info(FText::Format(LOCTEXT("LocalizationImportLog", "Localization Import Log: {LogMessage}"), Arguments));
+				}
+
+				FPlatformProcess::Sleep(0.25);
+				FPlatformProcess::ClosePipe(ReadPipe, WritePipe);
+
+				int32 ReturnCode;
+				if (!FPlatformProcess::GetProcReturnCode(ProcessHandle, &ReturnCode))
+				{
+					bHadError = true;
+				}
+				else if (ReturnCode != 0)
+				{
+					bHadError = true;
+				}
+			}
+
+			GWarn->EndSlowTask();
+
+			if (bHadError)
+			{
+				TranslationEditorMessageLog.Error(LOCTEXT("FailedToExportLocalization", "Failed to export localization!"));
+				TranslationEditorMessageLog.Notify(LOCTEXT("FailedToExportLocalization", "Failed to export localization!"), EMessageSeverity::Info, true);
+				TranslationEditorMessageLog.Open(EMessageSeverity::Error);
+			}
+			else
+			{
+				DataManager->LoadFromArchive(DataManager->GetAllTranslationsArray(), true, true);
+
+				TabManager->InvokeTab(ChangedOnImportTabId);
+				ChangedOnImportPropertyTable->SetObjects((TArray<UObject*>&)DataManager->GetChangedOnImportArray());
+				// Need to re-add the columns we want to display
+				ChangedOnImportPropertyTable->AddColumn((TWeakObjectPtr<UProperty>)FindField<UProperty>(UTranslationUnit::StaticClass(), "Source"));
+				ChangedOnImportPropertyTable->AddColumn((TWeakObjectPtr<UProperty>)FindField<UProperty>(UTranslationUnit::StaticClass(), "TranslationBeforeImport"));
+				ChangedOnImportPropertyTable->AddColumn((TWeakObjectPtr<UProperty>)FindField<UProperty>(UTranslationUnit::StaticClass(), "Translation"));
+			}
+		}
+	}
 }
 
 void FTranslationEditor::OnFilterTextChanged(const FText& InFilterText)
diff --git a/Engine/Source/Editor/TranslationEditor/Private/TranslationEditor.h b/Engine/Source/Editor/TranslationEditor/Private/TranslationEditor.h
--- a/Engine/Source/Editor/TranslationEditor/Private/TranslationEditor.h
+++ b/Engine/Source/Editor/TranslationEditor/Private/TranslationEditor.h
@@ -90,6 +90,8 @@ private:
 	TSharedRef<SDockTab> SpawnTab_History( const FSpawnTabArgs& Args );
 	/**	Spawns the search tab */
 	TSharedRef<SDockTab> SpawnTab_Search( const FSpawnTabArgs& Args );
+	/**	Spawns the Changed on Import tab */
+	TSharedRef<SDockTab> SpawnTab_ChangedOnImport(const FSpawnTabArgs& Args);
 
 	/** Map actions for the UI_COMMANDS */
 	void MapActions();
@@ -150,6 +152,9 @@ private:
 	/** Update content when a new translation unit selection is made in the Search PropertyTable */
 	void UpdateSearchSelection();
 
+	/** Update content when a new translation unit selection is made in the ChangedOnImport PropertyTable */
+	void UpdateChangedOnImportSelection();
+
 	/** Update content when a new context selection is made */
 	void UpdateContextSelection();
 
@@ -159,6 +164,9 @@ private:
 	/** Called when "Export to .PO" is clicked for this asset */
 	void ExportToPortableObjectFormat_Execute();
 
+	/** Called when "Import from .PO" is clicked for this asset */
+	void ImportFromPortableObjectFormat_Execute();
+
 	/** Open the search tab */
 	void OpenSearchTab_Execute();
 
@@ -182,6 +190,8 @@ private:
 	static const FName HistoryTabId;
 	/**	The tab id for the search tab */
 	static const FName SearchTabId;
+	/**	The tab id for the changed on import tab */
+	static const FName ChangedOnImportTabId;
 
 	/** The Untranslated Tab */
 	TWeakPtr<SDockTab> UntranslatedTab;
@@ -191,6 +201,8 @@ private:
 	TWeakPtr<SDockTab> CompletedTab;
 	/** The Search Tab */
 	TWeakPtr<SDockTab> SearchTab;
+	/** The Changed On Import Tab */
+	TWeakPtr<SDockTab> ChangedOnImportTab;
 
 	/** Search box for searching the source and translation strings */
 	TSharedPtr<SSearchBox> SearchBox;
@@ -212,6 +224,8 @@ private:
 	TSharedPtr< class IPropertyTable > HistoryPropertyTable;
 	/** The table of search results */
 	TSharedPtr< class IPropertyTable > SearchPropertyTable;
+	/** The table of changed on import results */
+	TSharedPtr< class IPropertyTable > ChangedOnImportPropertyTable;
 
 	/** The slate widget table of untranslated items */
 	TSharedPtr< class IPropertyTableWidgetHandle > UntranslatedPropertyTableWidgetHandle;
@@ -225,6 +239,8 @@ private:
 	TSharedPtr< class IPropertyTableWidgetHandle > HistoryPropertyTableWidgetHandle;
 	/** The slate widget table of search results */
 	TSharedPtr< class IPropertyTableWidgetHandle > SearchPropertyTableWidgetHandle;
+	/** The slate widget table of translations that changed on import */
+	TSharedPtr< class IPropertyTableWidgetHandle > ChangedOnImportPropertyTableWidgetHandle;
 
 	/** Font to use for the source language */
 	FSlateFontInfo SourceFont;
@@ -245,4 +261,9 @@ private:
 	FString ManifestFilePath;
 	/** Name of the language we are translating to */
 	FString ArchiveFilePath;
+
+	/** Used to remember the location of the file the user last exported to */
+	FString LastExportFilePath;
+	/** Used to remember the location of the file the user last imported */
+	FString LastImportFilePath;
 };
diff --git a/Engine/Source/Editor/TranslationEditor/Private/TranslationEditorMenu.cpp b/Engine/Source/Editor/TranslationEditor/Private/TranslationEditorMenu.cpp
--- a/Engine/Source/Editor/TranslationEditor/Private/TranslationEditorMenu.cpp
+++ b/Engine/Source/Editor/TranslationEditor/Private/TranslationEditorMenu.cpp
@@ -17,6 +17,7 @@ void FTranslationEditorMenu::FillTranslationMenu( FMenuBuilder& MenuBuilder/*, F
 		MenuBuilder.AddMenuEntry( FTranslationEditorCommands::Get().ChangeTranslationTargetFont );
 		MenuBuilder.AddMenuEntry( FTranslationEditorCommands::Get().PreviewAllTranslationsInEditor );
 		MenuBuilder.AddMenuEntry( FTranslationEditorCommands::Get().ExportToPortableObjectFormat );
+		MenuBuilder.AddMenuEntry( FTranslationEditorCommands::Get().ImportFromPortableObjectFormat );
 		MenuBuilder.AddMenuEntry( FTranslationEditorCommands::Get().OpenSearchTab );
 	}
 	MenuBuilder.EndSection();
@@ -70,6 +71,8 @@ void FTranslationEditorMenu::SetupTranslationEditorToolbar( TSharedPtr< FExtende
 				FTranslationEditorCommands::Get().PreviewAllTranslationsInEditor, "PreviewTranslationsInEditor", TAttribute<FText>(), TAttribute<FText>(), FSlateIcon(FEditorStyle::GetStyleSetName(), "AssetEditor.ReimportAsset")); 
 			ToolbarBuilder.AddToolBarButton(
 				FTranslationEditorCommands::Get().ExportToPortableObjectFormat, "ExportToPortableObjectFormat", TAttribute<FText>(), TAttribute<FText>(), FSlateIcon(FEditorStyle::GetStyleSetName(), "FontEditor.Export"));
+			ToolbarBuilder.AddToolBarButton(
+				FTranslationEditorCommands::Get().ImportFromPortableObjectFormat, "ImportFromPortableObjectFormat", TAttribute<FText>(), TAttribute<FText>(), FSlateIcon(FEditorStyle::GetStyleSetName(), "TranslationEditor.Import"));
 			ToolbarBuilder.AddToolBarButton(
 				FTranslationEditorCommands::Get().OpenSearchTab, "OpenSearchTab", TAttribute<FText>(), TAttribute<FText>(), FSlateIcon(FEditorStyle::GetStyleSetName(), "BlueprintEditor.FindInBlueprint"));
 		}
@@ -93,6 +96,7 @@ void FTranslationEditorCommands::RegisterCommands()
 	UI_COMMAND( SaveTranslations, "Save", "Saves the translations to file", EUserInterfaceActionType::Button, FInputGesture() );
 	UI_COMMAND( PreviewAllTranslationsInEditor, "Preview in Editor", "Preview All Translations in the Editor UI", EUserInterfaceActionType::Button, FInputGesture() );
 	UI_COMMAND( ExportToPortableObjectFormat, "Export to .PO", "Export to Portable Object Format", EUserInterfaceActionType::Button, FInputGesture() );
+	UI_COMMAND( ImportFromPortableObjectFormat, "Import from .PO", "Import from Portable Object Format", EUserInterfaceActionType::Button, FInputGesture() );
 	UI_COMMAND( OpenSearchTab, "Search", "Search Source and Translation Strings", EUserInterfaceActionType::Button, FInputGesture() );
 }
 
diff --git a/Engine/Source/Editor/TranslationEditor/Private/TranslationUnit.h b/Engine/Source/Editor/TranslationEditor/Private/TranslationUnit.h
--- a/Engine/Source/Editor/TranslationEditor/Private/TranslationUnit.h
+++ b/Engine/Source/Editor/TranslationEditor/Private/TranslationUnit.h
@@ -74,6 +74,10 @@ public:
 	UPROPERTY(Category=Translation, EditAnywhere)
 	bool HasBeenReviewed;
 
+	/** If this Translation Unit had a different translation before import, it will be stored here */
+	UPROPERTY(Category=Translation, EditAnywhere)
+	FString TranslationBeforeImport;
+
 	/**
 	* Returns an event delegate that is executed when a property has changed.
 	*
diff --git a/Engine/Source/Editor/TranslationEditor/Public/TranslationEditorMenu.h b/Engine/Source/Editor/TranslationEditor/Public/TranslationEditorMenu.h
--- a/Engine/Source/Editor/TranslationEditor/Public/TranslationEditorMenu.h
+++ b/Engine/Source/Editor/TranslationEditor/Public/TranslationEditorMenu.h
@@ -39,6 +39,9 @@ public:
 	/** Export to PortableObject format (.po) */
 	TSharedPtr<FUICommandInfo> ExportToPortableObjectFormat;
 
+	/** Import from PortableObject format (.po) */
+	TSharedPtr<FUICommandInfo> ImportFromPortableObjectFormat;
+
 	/** Open the tab for searching */
 	TSharedPtr<FUICommandInfo> OpenSearchTab;
 
