fix non-unity build

diff --git a/Engine/Source/Editor/TranslationEditor/Private/TranslationUnit.cpp b/Engine/Source/Editor/TranslationEditor/Private/TranslationUnit.cpp
--- a/Engine/Source/Editor/TranslationEditor/Private/TranslationUnit.cpp
+++ b/Engine/Source/Editor/TranslationEditor/Private/TranslationUnit.cpp
@@ -1,5 +1,6 @@
 // Copyright 1998-2014 Epic Games, Inc. All Rights Reserved.
 #include "TranslationEditorPrivatePCH.h"
+#include "TranslationUnit.h"
 
 UTranslationUnit::UTranslationUnit( const class FPostConstructInitializeProperties& PCIP )
 	: Super(PCIP)
