Translation Editor & Internationalization Export Commandlet fixes:

Translation Editor:
- Added some NULL checks on UTranslationUnit objects
- Added setting for PortaleObjectName for export/import

Internationlization Export:
- Fixed handling of Filename
- Make sure the strings importing aren't marked as changed if they only contain whitespace

diff --git a/Engine/Source/Editor/TranslationEditor/Private/InternationalizationExportSettings.h b/Engine/Source/Editor/TranslationEditor/Private/InternationalizationExportSettings.h
--- a/Engine/Source/Editor/TranslationEditor/Private/InternationalizationExportSettings.h
+++ b/Engine/Source/Editor/TranslationEditor/Private/InternationalizationExportSettings.h
@@ -26,6 +26,10 @@ public:
 	UPROPERTY(Category = GatherTextStep, EditAnywhere, config)
 	FString DestinationPath;
 
+	/** Filename for the Portable Object format file */
+	UPROPERTY(Category = GatherTextStep, EditAnywhere, config)
+	FString PortableObjectName;
+
 	/** Name of the manifest file */
 	UPROPERTY(Category = GatherTextStep, EditAnywhere, config)
 	FString ManifestName;
diff --git a/Engine/Source/Editor/TranslationEditor/Private/TranslationDataManager.cpp b/Engine/Source/Editor/TranslationEditor/Private/TranslationDataManager.cpp
--- a/Engine/Source/Editor/TranslationEditor/Private/TranslationDataManager.cpp
+++ b/Engine/Source/Editor/TranslationEditor/Private/TranslationDataManager.cpp
@@ -169,37 +169,46 @@ bool FTranslationDataManager::WriteTranslationData(bool bForceWrite /*= false*/)
 
 	for (UTranslationUnit* TranslationUnit : Untranslated)
 	{
-		const FLocItem SearchSource(TranslationUnit->Source);
-		FString OldTranslation = Archive->FindEntryBySource(TranslationUnit->Namespace, SearchSource, NULL)->Translation.Text;
-		FString TranslationToWrite = TranslationUnit->Translation.ReplaceCharWithEscapedChar();
-		if ( !TranslationToWrite.Equals(OldTranslation) )
+		if (TranslationUnit != NULL)
 		{
-			Archive->SetTranslation(TranslationUnit->Namespace, TranslationUnit->Source, TranslationToWrite, NULL);
-			bNeedsWrite = true;
+			const FLocItem SearchSource(TranslationUnit->Source);
+			FString OldTranslation = Archive->FindEntryBySource(TranslationUnit->Namespace, SearchSource, NULL)->Translation.Text;
+			FString TranslationToWrite = TranslationUnit->Translation.ReplaceCharWithEscapedChar();
+			if (!TranslationToWrite.Equals(OldTranslation))
+			{
+				Archive->SetTranslation(TranslationUnit->Namespace, TranslationUnit->Source, TranslationToWrite, NULL);
+				bNeedsWrite = true;
+			}
 		}
 	}
 
 	for (UTranslationUnit* TranslationUnit : Review)
 	{
-		const FLocItem SearchSource(TranslationUnit->Source);
-		FString OldTranslation = Archive->FindEntryBySource(TranslationUnit->Namespace, SearchSource, NULL)->Translation.Text;
-		FString TranslationToWrite = TranslationUnit->Translation.ReplaceCharWithEscapedChar();
-		if ( TranslationUnit->HasBeenReviewed && !TranslationToWrite.Equals(OldTranslation) )
+		if (TranslationUnit != NULL)
 		{
-			Archive->SetTranslation(TranslationUnit->Namespace, TranslationUnit->Source, TranslationToWrite, NULL);
-			bNeedsWrite = true;
+			const FLocItem SearchSource(TranslationUnit->Source);
+			FString OldTranslation = Archive->FindEntryBySource(TranslationUnit->Namespace, SearchSource, NULL)->Translation.Text;
+			FString TranslationToWrite = TranslationUnit->Translation.ReplaceCharWithEscapedChar();
+			if (TranslationUnit->HasBeenReviewed && !TranslationToWrite.Equals(OldTranslation))
+			{
+				Archive->SetTranslation(TranslationUnit->Namespace, TranslationUnit->Source, TranslationToWrite, NULL);
+				bNeedsWrite = true;
+			}
 		}
 	}
 
 	for (UTranslationUnit* TranslationUnit : Complete)
 	{
-		const FLocItem SearchSource(TranslationUnit->Source);
-		FString OldTranslation = Archive->FindEntryBySource(TranslationUnit->Namespace, SearchSource, NULL)->Translation.Text;
-		FString TranslationToWrite = TranslationUnit->Translation.ReplaceCharWithEscapedChar();
-		if ( !TranslationToWrite.Equals(OldTranslation) )
+		if (TranslationUnit != NULL)
 		{
-			Archive->SetTranslation(TranslationUnit->Namespace, TranslationUnit->Source, TranslationToWrite, NULL);
-			bNeedsWrite = true;
+			const FLocItem SearchSource(TranslationUnit->Source);
+			FString OldTranslation = Archive->FindEntryBySource(TranslationUnit->Namespace, SearchSource, NULL)->Translation.Text;
+			FString TranslationToWrite = TranslationUnit->Translation.ReplaceCharWithEscapedChar();
+			if (!TranslationToWrite.Equals(OldTranslation))
+			{
+				Archive->SetTranslation(TranslationUnit->Namespace, TranslationUnit->Source, TranslationToWrite, NULL);
+				bNeedsWrite = true;
+			}
 		}
 	}
 
@@ -395,7 +404,7 @@ void FTranslationDataManager::GetHistoryForTranslationUnits( TArray<UTranslation
 
 					for (UTranslationUnit* TranslationUnit : TranslationUnits)
 					{
-						if(TranslationUnit->Contexts.Num() > 0)
+						if(TranslationUnit != NULL && TranslationUnit->Contexts.Num() > 0)
 						{
 							for (FTranslationContextInfo& ContextInfo : TranslationUnit->Contexts)
 							{
@@ -496,10 +505,13 @@ void FTranslationDataManager::PopulateSearchResultsUsingFilter(const FString& Se
 
 	for (UTranslationUnit* TranslationUnit : AllTranslations)
 	{
-		if (TranslationUnit->Source.Contains(SearchFilter) || 
-			TranslationUnit->Translation.Contains(SearchFilter))
+		if (TranslationUnit != NULL)
 		{
-			SearchResults.Add(TranslationUnit);
+			if (TranslationUnit->Source.Contains(SearchFilter) ||
+				TranslationUnit->Translation.Contains(SearchFilter))
+			{
+				SearchResults.Add(TranslationUnit);
+			}
 		}
 	}
 }
@@ -531,67 +543,70 @@ void FTranslationDataManager::LoadFromArchive(TArray<UTranslationUnit*>& InTrans
 		for (int32 CurrentTranslationUnitIndex = 0; CurrentTranslationUnitIndex < TranslationUnits.Num(); ++CurrentTranslationUnitIndex)
 		{
 			UTranslationUnit* TranslationUnit = TranslationUnits[CurrentTranslationUnitIndex];
-			AllTranslations.Add(TranslationUnit);
-
-			GWarn->StatusUpdate(CurrentTranslationUnitIndex, TranslationUnits.Num(), FText::Format(LOCTEXT("LoadingCurrentArchiveEntries", "Loading Entry {0} of {1} from Translation Archive..."), FText::AsNumber(CurrentTranslationUnitIndex), FText::AsNumber(TranslationUnits.Num())));
-
-			const FLocItem SourceSearch(TranslationUnit->Source);
-			TSharedPtr<FArchiveEntry> ArchiveEntry = Archive->FindEntryBySource(TranslationUnit->Namespace, SourceSearch, NULL);
-			if (ArchiveEntry.IsValid())
+			if (TranslationUnit != NULL)
 			{
-				const FString PreviousTranslation = TranslationUnit->Translation;
-				TranslationUnit->Translation = ""; // Reset to null string
-				const FString UnescapedTranslatedString = ArchiveEntry->Translation.Text.ReplaceEscapedCharWithChar();
+				AllTranslations.Add(TranslationUnit);
 
-				if (UnescapedTranslatedString.IsEmpty())
+				GWarn->StatusUpdate(CurrentTranslationUnitIndex, TranslationUnits.Num(), FText::Format(LOCTEXT("LoadingCurrentArchiveEntries", "Loading Entry {0} of {1} from Translation Archive..."), FText::AsNumber(CurrentTranslationUnitIndex), FText::AsNumber(TranslationUnits.Num())));
+
+				const FLocItem SourceSearch(TranslationUnit->Source);
+				TSharedPtr<FArchiveEntry> ArchiveEntry = Archive->FindEntryBySource(TranslationUnit->Namespace, SourceSearch, NULL);
+				if (ArchiveEntry.IsValid())
 				{
-					bool bHasTranslationHistory = false;
-					int32 MostRecentNonNullTranslationIndex = -1;
-					int32 ContextForRecentTranslation = -1;
+					const FString PreviousTranslation = TranslationUnit->Translation;
+					TranslationUnit->Translation = ""; // Reset to null string
+					const FString UnescapedTranslatedString = ArchiveEntry->Translation.Text.ReplaceEscapedCharWithChar();
 
-					for (int32 ContextIndex = 0; ContextIndex < TranslationUnit->Contexts.Num(); ++ContextIndex)
+					if (UnescapedTranslatedString.IsEmpty())
 					{
-						for (int32 ChangeIndex = 0; ChangeIndex < TranslationUnit->Contexts[ContextIndex].Changes.Num(); ++ChangeIndex)
+						bool bHasTranslationHistory = false;
+						int32 MostRecentNonNullTranslationIndex = -1;
+						int32 ContextForRecentTranslation = -1;
+
+						for (int32 ContextIndex = 0; ContextIndex < TranslationUnit->Contexts.Num(); ++ContextIndex)
 						{
-							if (!(TranslationUnit->Contexts[ContextIndex].Changes[ChangeIndex].Translation.IsEmpty()))
+							for (int32 ChangeIndex = 0; ChangeIndex < TranslationUnit->Contexts[ContextIndex].Changes.Num(); ++ChangeIndex)
+							{
+								if (!(TranslationUnit->Contexts[ContextIndex].Changes[ChangeIndex].Translation.IsEmpty()))
+								{
+									bHasTranslationHistory = true;
+									MostRecentNonNullTranslationIndex = ChangeIndex;
+									ContextForRecentTranslation = ContextIndex;
+									break;
+								}
+							}
+
+							if (bHasTranslationHistory)
 							{
-								bHasTranslationHistory = true;
-								MostRecentNonNullTranslationIndex = ChangeIndex;
-								ContextForRecentTranslation = ContextIndex;
 								break;
 							}
 						}
 
+						// If we have history, but current translation is empty, this goes in the Needs Review tab
 						if (bHasTranslationHistory)
 						{
-							break;
+							// Offer the most recent translation (for the first context in the list) as a suggestion or starting point (not saved unless user checks "Has Been Reviewed")
+							TranslationUnit->Translation = TranslationUnit->Contexts[ContextForRecentTranslation].Changes[MostRecentNonNullTranslationIndex].Translation;
+							Review.Add(TranslationUnit);
+						}
+						else
+						{
+							Untranslated.Add(TranslationUnit);
 						}
-					}
-
-					// If we have history, but current translation is empty, this goes in the Needs Review tab
-					if (bHasTranslationHistory)
-					{
-						// Offer the most recent translation (for the first context in the list) as a suggestion or starting point (not saved unless user checks "Has Been Reviewed")
-						TranslationUnit->Translation = TranslationUnit->Contexts[ContextForRecentTranslation].Changes[MostRecentNonNullTranslationIndex].Translation;
-						Review.Add(TranslationUnit);
 					}
 					else
 					{
-						Untranslated.Add(TranslationUnit);
+						TranslationUnit->Translation = UnescapedTranslatedString;
+						TranslationUnit->HasBeenReviewed = true;
+						Complete.Add(TranslationUnit);
 					}
-				}
-				else
-				{
-					TranslationUnit->Translation = UnescapedTranslatedString;
-					TranslationUnit->HasBeenReviewed = true;
-					Complete.Add(TranslationUnit);
-				}
 
-				// Add to changed array if we're tracking changes
-				if (bTrackChanges && PreviousTranslation != TranslationUnit->Translation)
-				{
-					ChangedOnImport.Add(TranslationUnit);
-					TranslationUnit->TranslationBeforeImport = PreviousTranslation;
+					// Add to changed array if we're tracking changes
+					if (bTrackChanges && PreviousTranslation != TranslationUnit->Translation)
+					{
+						ChangedOnImport.Add(TranslationUnit);
+						TranslationUnit->TranslationBeforeImport = PreviousTranslation;
+					}
 				}
 			}
 		}
diff --git a/Engine/Source/Editor/TranslationEditor/Private/TranslationEditor.cpp b/Engine/Source/Editor/TranslationEditor/Private/TranslationEditor.cpp
--- a/Engine/Source/Editor/TranslationEditor/Private/TranslationEditor.cpp
+++ b/Engine/Source/Editor/TranslationEditor/Private/TranslationEditor.cpp
@@ -152,7 +152,7 @@ void FTranslationEditor::InitTranslationEditor( const EToolkitMode::Type Mode, c
 	const bool bCreateDefaultToolbar = true;
 	// Need editing object to not be null
 	UTranslationUnit* EditingObject;
-	if (DataManager->GetAllTranslationsArray().Num() > 0)
+	if (DataManager->GetAllTranslationsArray().Num() > 0 && DataManager->GetAllTranslationsArray()[0] != NULL)
 	{
 		EditingObject = DataManager->GetAllTranslationsArray()[0];
 	}
@@ -1125,7 +1125,8 @@ void FTranslationEditor::ExportToPortableObjectFormat_Execute()
 
 		if (SaveFilenames.Num() > 0)
 		{
-			ExportSettings->DestinationPath = SaveFilenames[0];
+			ExportSettings->DestinationPath = FPaths::GetPath(SaveFilenames[0]);
+			ExportSettings->PortableObjectName = FPaths::GetCleanFilename(SaveFilenames[0]);
 			LastExportFilePath = FPaths::GetPath(SaveFilenames[0]);
 		}
 
@@ -1248,7 +1249,8 @@ void FTranslationEditor::ImportFromPortableObjectFormat_Execute()
 
 		if (OpenFilenames.Num() > 0)
 		{
-			ImportSettings->SourcePath = OpenFilenames[0];
+			ImportSettings->SourcePath = FPaths::GetPath(OpenFilenames[0]);
+			ImportSettings->PortableObjectName = FPaths::GetCleanFilename(OpenFilenames[0]);
 			LastImportFilePath = FPaths::GetPath(OpenFilenames[0]);
 		}
 
diff --git a/Engine/Source/Editor/UnrealEd/Private/Commandlets/InternationalizationExportCommandlet.cpp b/Engine/Source/Editor/UnrealEd/Private/Commandlets/InternationalizationExportCommandlet.cpp
--- a/Engine/Source/Editor/UnrealEd/Private/Commandlets/InternationalizationExportCommandlet.cpp
+++ b/Engine/Source/Editor/UnrealEd/Private/Commandlets/InternationalizationExportCommandlet.cpp
@@ -1066,15 +1066,7 @@ bool UInternationalizationExportCommandlet::DoExport( const FString& SourcePath,
 				// Write out the Portable Object to .po file.
 				{
 					FString OutputString = PortableObj.ToString();
-					FString OutputFileName;
-					if (FPaths::GetExtension(DestinationPath) != "po")
-					{
-						OutputFileName = DestinationPath / CultureName / Filename;
-					}
-					else
-					{
-						OutputFileName = DestinationPath;
-					}
+					FString OutputFileName = DestinationPath / Filename;
 
 					//@TODO We force UTF8 at the moment but we want this to be based on the format found in the header info.
 					if( !FFileHelper::SaveStringToFile(OutputString, *OutputFileName, FFileHelper::EEncodingOptions::ForceUTF8) )
@@ -1112,15 +1104,7 @@ bool UInternationalizationExportCommandlet::DoImport(const FString& SourcePath,
 	{
 		// Load the Portable Object file if found
 		const FString CultureName = CulturesToGenerate[Culture];
-		FString POFilePath;
-		if (FPaths::GetExtension(SourcePath) != "po")
-		{
-			POFilePath = SourcePath / CultureName / Filename;
-		}
-		else
-		{
-			POFilePath = SourcePath;
-		}
+		FString POFilePath = SourcePath / Filename;
 
 		if( !FPaths::FileExists(POFilePath) )
 		{
@@ -1169,7 +1153,7 @@ bool UInternationalizationExportCommandlet::DoImport(const FString& SourcePath,
 			for( auto EntryIter = PortableObject.GetEntriesIterator(); EntryIter; ++EntryIter )
 			{
 				auto POEntry = *EntryIter;
-				if( POEntry->MsgId.IsEmpty() || POEntry->MsgStr.Num() == 0 || POEntry->MsgStr[0].IsEmpty() )
+				if( POEntry->MsgId.IsEmpty() || POEntry->MsgStr.Num() == 0 || POEntry->MsgStr[0].Trim().IsEmpty() )
 				{
 					// We ignore the header entry or entries with no translation.
 					continue;
