In the Translation Editor, UTranslationUnit objects should be added to the root set on creation to prevent them from being garbage collected, then deleted from the root set when the translation editor is closed to allow them to be garbage collected at that time.

diff --git a/Engine/Source/Editor/TranslationEditor/Private/TranslationDataManager.cpp b/Engine/Source/Editor/TranslationEditor/Private/TranslationDataManager.cpp
--- a/Engine/Source/Editor/TranslationEditor/Private/TranslationDataManager.cpp
+++ b/Engine/Source/Editor/TranslationEditor/Private/TranslationDataManager.cpp
@@ -100,6 +100,11 @@ FTranslationDataManager::FTranslationDataManager( const FString& InManifestFileP
 	GWarn->EndSlowTask();
 }
 
+FTranslationDataManager::~FTranslationDataManager()
+{
+	RemoveTranslationUnitArrayfromRoot(AllTranslations); // Re-enable garbage collection for all current UTranslationDataObjects
+}
+
 TSharedPtr< FInternationalizationManifest > FTranslationDataManager::ReadManifest( const FString& ManifestFilePathToRead )
 {
 	
@@ -539,12 +544,15 @@ void FTranslationDataManager::LoadFromArchive(TArray<UTranslationUnit*>& InTrans
 		Complete.Empty();
 		ChangedOnImport.Empty();
 
-
 		for (int32 CurrentTranslationUnitIndex = 0; CurrentTranslationUnitIndex < TranslationUnits.Num(); ++CurrentTranslationUnitIndex)
 		{
 			UTranslationUnit* TranslationUnit = TranslationUnits[CurrentTranslationUnitIndex];
 			if (TranslationUnit != NULL)
 			{
+				if (!TranslationUnit->IsRooted())
+				{
+					TranslationUnit->AddToRoot(); // Disable garbage collection for UTranslationUnit objects
+				}
 				AllTranslations.Add(TranslationUnit);
 
 				GWarn->StatusUpdate(CurrentTranslationUnitIndex, TranslationUnits.Num(), FText::Format(LOCTEXT("LoadingCurrentArchiveEntries", "Loading Entry {0} of {1} from Translation Archive..."), FText::AsNumber(CurrentTranslationUnitIndex), FText::AsNumber(TranslationUnits.Num())));
@@ -624,4 +632,12 @@ void FTranslationDataManager::LoadFromArchive(TArray<UTranslationUnit*>& InTrans
 	GWarn->EndSlowTask();
 }
 
+void FTranslationDataManager::RemoveTranslationUnitArrayfromRoot(TArray<UTranslationUnit*>& TranslationUnits)
+{
+	for (UTranslationUnit* TranslationUnit : TranslationUnits)
+	{
+		TranslationUnit->RemoveFromRoot();
+	}
+}
+
 #undef LOCTEXT_NAMESPACE
diff --git a/Engine/Source/Editor/TranslationEditor/Private/TranslationDataManager.h b/Engine/Source/Editor/TranslationEditor/Private/TranslationDataManager.h
--- a/Engine/Source/Editor/TranslationEditor/Private/TranslationDataManager.h
+++ b/Engine/Source/Editor/TranslationEditor/Private/TranslationDataManager.h
@@ -17,6 +17,8 @@ public:
 
 	FTranslationDataManager( const FString& InManifestFilePath, const FString& InArchiveFilePath);
 
+	virtual ~FTranslationDataManager();
+
 	TArray<UTranslationUnit*>& GetAllTranslationsArray()
 	{
 		return AllTranslations;
@@ -80,6 +82,9 @@ private:
 	/** Get the history data for a given translation unit */
 	void GetHistoryForTranslationUnits( TArray<UTranslationUnit*>& TranslationUnits, const FString& ManifestFilePath );
 
+	/** Removes each UTranslationUnit in the passed array from the root set, allowing it to be garbage collected */
+	void RemoveTranslationUnitArrayfromRoot( TArray<UTranslationUnit*>& TranslationUnits );
+
 	// Arrays containing the translation data
 	TArray<UTranslationUnit*> AllTranslations;
 	TArray<UTranslationUnit*> Untranslated;
