Translation Editor import from .po : ignore changes to leading and trailing whitespace

diff --git a/Engine/Source/Editor/TranslationEditor/Private/TranslationDataManager.cpp b/Engine/Source/Editor/TranslationEditor/Private/TranslationDataManager.cpp
--- a/Engine/Source/Editor/TranslationEditor/Private/TranslationDataManager.cpp
+++ b/Engine/Source/Editor/TranslationEditor/Private/TranslationDataManager.cpp
@@ -609,11 +609,26 @@ void FTranslationDataManager::LoadFromArchive(TArray<UTranslationUnit*>& InTrans
 						Complete.Add(TranslationUnit);
 					}
 
-					// Add to changed array if we're tracking changes
-					if (bTrackChanges && PreviousTranslation != TranslationUnit->Translation)
+					// Add to changed array if we're tracking changes (i.e. when we import from .po files)
+					if (bTrackChanges)
 					{
-						ChangedOnImport.Add(TranslationUnit);
-						TranslationUnit->TranslationBeforeImport = PreviousTranslation;
+						if (PreviousTranslation != TranslationUnit->Translation)
+						{
+							FString PreviousTranslationTrimmed = PreviousTranslation;
+							PreviousTranslationTrimmed.Trim().TrimTrailing();
+							FString CurrentTranslationTrimmed = TranslationUnit->Translation;
+							CurrentTranslationTrimmed.Trim().TrimTrailing();
+							// Ignore changes to only whitespace at beginning and/or end of string on import
+							if (PreviousTranslationTrimmed == CurrentTranslationTrimmed)
+							{
+								TranslationUnit->Translation = PreviousTranslation;
+							}
+							else
+							{
+								ChangedOnImport.Add(TranslationUnit);
+								TranslationUnit->TranslationBeforeImport = PreviousTranslation;
+							}
+						}
 					}
 				}
 			}
