Shadow variable fixes

diff --git a/Engine/Source/Editor/BlueprintGraph/Private/K2ActionMenuBuilder.cpp b/Engine/Source/Editor/BlueprintGraph/Private/K2ActionMenuBuilder.cpp
--- a/Engine/Source/Editor/BlueprintGraph/Private/K2ActionMenuBuilder.cpp
+++ b/Engine/Source/Editor/BlueprintGraph/Private/K2ActionMenuBuilder.cpp
@@ -296,7 +296,7 @@ static void GetEnumUtilitiesNodes(FBlueprintGraphActionListBuilder& ContextMenuB
 {
 	struct FGetEnumUtilitiesNodes
 	{
-		static void Get(FBlueprintGraphActionListBuilder& InContextMenuBuilder, bool bInNumEnum, bool bInForEach, bool bInCastFromByte, bool bLiteralByte, UEnum* Enum, const FString& Category)
+		static void Get(FBlueprintGraphActionListBuilder& InContextMenuBuilder, bool bInNumEnum, bool bInForEach, bool bInCastFromByte, bool bInLiteralByte, UEnum* Enum, const FString& Category)
 		{
 			const bool bIsBlueprintType = UEdGraphSchema_K2::IsAllowableBlueprintVariableType(Enum);
 			if (bIsBlueprintType)
@@ -335,7 +335,7 @@ static void GetEnumUtilitiesNodes(FBlueprintGraphActionListBuilder& ContextMenuB
 					Action->SearchTitle = NodeTemplate->GetNodeSearchTitle();
 				}
 
-				if (bLiteralByte)
+				if (bInLiteralByte)
 				{
 					UK2Node_EnumLiteral* EnumTemplate = InContextMenuBuilder.CreateTemplateNode<UK2Node_EnumLiteral>();
 					EnumTemplate->Enum = Enum;
diff --git a/Engine/Source/Editor/BlueprintGraph/Private/K2Node_CallFunction.cpp b/Engine/Source/Editor/BlueprintGraph/Private/K2Node_CallFunction.cpp
--- a/Engine/Source/Editor/BlueprintGraph/Private/K2Node_CallFunction.cpp
+++ b/Engine/Source/Editor/BlueprintGraph/Private/K2Node_CallFunction.cpp
@@ -1220,12 +1220,12 @@ void UK2Node_CallFunction::ExpandNode(class FKismetCompilerContext& CompilerCont
 
 				struct FStructConnectHelper
 				{
-					static void Connect(const FString& PinName, UK2Node* Node, UEdGraphPin* BetterSelf, const UEdGraphSchema_K2* Schema, FCompilerResultsLog& MessageLog)
+					static void Connect(const FString& PinName, UK2Node* Node, UEdGraphPin* BetterSelf, const UEdGraphSchema_K2* InSchema, FCompilerResultsLog& MessageLog)
 					{
 						auto Pin = Node->FindPin(PinName);
 						if (!PinName.IsEmpty() && Pin && !Pin->LinkedTo.Num())
 						{
-							const bool bConnected = Schema->TryCreateConnection(Pin, BetterSelf);
+							const bool bConnected = InSchema->TryCreateConnection(Pin, BetterSelf);
 							if (!bConnected)
 							{
 								MessageLog.Warning(*LOCTEXT("DefaultToSelfNotConnected", "DefaultToSelf pin @@ from node @@ cannot be connected to @@").ToString(), Pin, Node, BetterSelf);
diff --git a/Engine/Source/Editor/LandscapeEditor/Private/LandscapeEdModeComponentTools.cpp b/Engine/Source/Editor/LandscapeEditor/Private/LandscapeEdModeComponentTools.cpp
--- a/Engine/Source/Editor/LandscapeEditor/Private/LandscapeEdModeComponentTools.cpp
+++ b/Engine/Source/Editor/LandscapeEditor/Private/LandscapeEdModeComponentTools.cpp
@@ -1164,7 +1164,7 @@ public:
 					{
 						for (int32 i = -1; (!bApplyToAll && i < 0) || i < LayerNum; ++i )
 						{
-							FGizmoPreData GizmoData[NeighborNum];
+							FGizmoPreData GizmoPreData[NeighborNum];
 
 							for (int32 LocalY = 0; LocalY < 2; ++LocalY)
 							{
@@ -1172,18 +1172,18 @@ public:
 								{
 									int32 x = FMath::Clamp(LX + LocalX, X1, X2);
 									int32 y = FMath::Clamp(LY + LocalY, Y1, Y2);
-									GizmoData[LocalX + LocalY*2].Ratio = LandscapeInfo->SelectedRegion.FindRef(FIntPoint(x, y));
+									GizmoPreData[LocalX + LocalY * 2].Ratio = LandscapeInfo->SelectedRegion.FindRef(FIntPoint(x, y));
 									int32 index = (x-X1) + (y-Y1)*(1+X2-X1);
 
 									if (bApplyToAll)
 									{
 										if ( i < 0 )
 										{
-											GizmoData[LocalX + LocalY*2].Data = Gizmo->GetNormalizedHeight( HeightData[index] );
+											GizmoPreData[LocalX + LocalY * 2].Data = Gizmo->GetNormalizedHeight(HeightData[index]);
 										}
 										else
 										{
-											GizmoData[LocalX + LocalY*2].Data = WeightDatas[index*LayerNum + i ];
+											GizmoPreData[LocalX + LocalY * 2].Data = WeightDatas[index*LayerNum + i];
 										}
 									}
 									else
@@ -1191,11 +1191,11 @@ public:
 										typename ToolTarget::CacheClass::DataType OriginalValue = Data[index];
 										if ( EdMode->CurrentToolTarget.TargetType == ELandscapeToolTargetType::Heightmap )
 										{
-											GizmoData[LocalX + LocalY*2].Data = Gizmo->GetNormalizedHeight(OriginalValue);
+											GizmoPreData[LocalX + LocalY * 2].Data = Gizmo->GetNormalizedHeight(OriginalValue);
 										}
 										else
 										{
-											GizmoData[LocalX + LocalY*2].Data = OriginalValue;
+											GizmoPreData[LocalX + LocalY * 2].Data = OriginalValue;
 										}
 									}
 								}
@@ -1206,14 +1206,14 @@ public:
 							float FracY = LandscapeLocal.Y - LY;
 							LerpedData.Ratio = bFullCopy ? 1.f : 
 								FMath::Lerp(
-								FMath::Lerp(GizmoData[0].Ratio, GizmoData[1].Ratio, FracX),
-								FMath::Lerp(GizmoData[2].Ratio, GizmoData[3].Ratio, FracX),
+								FMath::Lerp(GizmoPreData[0].Ratio, GizmoPreData[1].Ratio, FracX),
+								FMath::Lerp(GizmoPreData[2].Ratio, GizmoPreData[3].Ratio, FracX),
 								FracY
 								);
 
 							LerpedData.Data = FMath::Lerp(
-								FMath::Lerp(GizmoData[0].Data, GizmoData[1].Data, FracX),
-								FMath::Lerp(GizmoData[2].Data, GizmoData[3].Data, FracX),
+								FMath::Lerp(GizmoPreData[0].Data, GizmoPreData[1].Data, FracX),
+								FMath::Lerp(GizmoPreData[2].Data, GizmoPreData[3].Data, FracX),
 								FracY
 								);
 
@@ -1240,29 +1240,29 @@ public:
 									}
 								}
 
-								FGizmoSelectData* GizmoData = Gizmo->SelectedData.Find(ALandscape::MakeKey(X, Y));
-								if (GizmoData)
+								FGizmoSelectData* GizmoSelectData = Gizmo->SelectedData.Find(ALandscape::MakeKey(X, Y));
+								if (GizmoSelectData)
 								{
 									if (bApplyToAll)
 									{
 										if (i < 0)
 										{
-											GizmoData->HeightData = LerpedData.Data;
+											GizmoSelectData->HeightData = LerpedData.Data;
 										}
 										else
 										{
-											GizmoData->WeightDataMap.Add(LandscapeInfo->Layers[i].LayerInfoObj, LerpedData.Data);
+											GizmoSelectData->WeightDataMap.Add(LandscapeInfo->Layers[i].LayerInfoObj, LerpedData.Data);
 										}
 									}
 									else
 									{
 										if (EdMode->CurrentToolTarget.TargetType == ELandscapeToolTargetType::Heightmap)
 										{
-											GizmoData->HeightData = LerpedData.Data;
+											GizmoSelectData->HeightData = LerpedData.Data;
 										}
 										else
 										{
-											GizmoData->WeightDataMap.Add(EdMode->CurrentToolTarget.LayerInfo.Get(), LerpedData.Data);
+											GizmoSelectData->WeightDataMap.Add(EdMode->CurrentToolTarget.LayerInfo.Get(), LerpedData.Data);
 										}
 									}
 								}
diff --git a/Engine/Source/Editor/LandscapeEditor/Private/LandscapeEdModeSplineTools.cpp b/Engine/Source/Editor/LandscapeEditor/Private/LandscapeEdModeSplineTools.cpp
--- a/Engine/Source/Editor/LandscapeEditor/Private/LandscapeEdModeSplineTools.cpp
+++ b/Engine/Source/Editor/LandscapeEditor/Private/LandscapeEdModeSplineTools.cpp
@@ -1686,8 +1686,8 @@ void FEdModeLandscape::SplineMoveToCurrentLevel()
 
 			for (ULandscapeSplineControlPoint* ControlPoint : SplineTool->SelectedSplineControlPoints)
 			{
-				ULandscapeSplinesComponent* Comp = ControlPoint->GetOuterULandscapeSplinesComponent();
-				ALandscapeProxy* FromProxy = Comp ? Cast<ALandscapeProxy>(Comp->GetOuter()) : NULL;
+				ULandscapeSplinesComponent* LandscapeSplinesComp = ControlPoint->GetOuterULandscapeSplinesComponent();
+				ALandscapeProxy* FromProxy = LandscapeSplinesComp ? Cast<ALandscapeProxy>(LandscapeSplinesComp->GetOuter()) : NULL;
 				if (FromProxy)
 				{
 					if (!Landscape)
@@ -1731,13 +1731,13 @@ void FEdModeLandscape::SplineMoveToCurrentLevel()
 
 						if (ControlPoint->MeshComponent)
 						{
-							UControlPointMeshComponent* Comp = ControlPoint->MeshComponent;
-							Comp->Modify();
-							Comp->UnregisterComponent();
-							Comp->DetachFromParent(true);
-							Comp->InvalidateLightingCache();
-							Comp->Rename(NULL, Landscape);
-							Comp->AttachTo(Landscape->SplineComponent, NAME_None, EAttachLocation::KeepWorldPosition);
+							UControlPointMeshComponent* ControlPointMeshComp = ControlPoint->MeshComponent;
+							ControlPointMeshComp->Modify();
+							ControlPointMeshComp->UnregisterComponent();
+							ControlPointMeshComp->DetachFromParent(true);
+							ControlPointMeshComp->InvalidateLightingCache();
+							ControlPointMeshComp->Rename(NULL, Landscape);
+							ControlPointMeshComp->AttachTo(Landscape->SplineComponent, NAME_None, EAttachLocation::KeepWorldPosition);
 						}
 
 						ControlPoint->UpdateSplinePoints(true, false);
@@ -1747,8 +1747,8 @@ void FEdModeLandscape::SplineMoveToCurrentLevel()
 
 			for (ULandscapeSplineSegment* Segment : SplineTool->SelectedSplineSegments)
 			{
-				ULandscapeSplinesComponent* Comp = Segment->GetOuterULandscapeSplinesComponent();
-				ALandscapeProxy* FromProxy = Comp ? Cast<ALandscapeProxy>(Comp->GetOuter()) : NULL;
+				ULandscapeSplinesComponent* LandscapeSplinesComp = Segment->GetOuterULandscapeSplinesComponent();
+				ALandscapeProxy* FromProxy = LandscapeSplinesComp ? Cast<ALandscapeProxy>(LandscapeSplinesComp->GetOuter()) : NULL;
 				if (FromProxy)
 				{
 					if (!Landscape)
@@ -1785,14 +1785,14 @@ void FEdModeLandscape::SplineMoveToCurrentLevel()
 						Segment->Rename(NULL, Landscape->SplineComponent);
 						Landscape->SplineComponent->Segments.Add(Segment);
 
-						for (USplineMeshComponent* Comp : Segment->MeshComponents)
+						for (USplineMeshComponent* SplineMeshComp : Segment->MeshComponents)
 						{
-							Comp->Modify();
-							Comp->UnregisterComponent();
-							Comp->DetachFromParent(true);
-							Comp->InvalidateLightingCache();
-							Comp->Rename(NULL, Landscape);
-							Comp->AttachTo(Landscape->SplineComponent, NAME_None, EAttachLocation::KeepWorldPosition);
+							SplineMeshComp->Modify();
+							SplineMeshComp->UnregisterComponent();
+							SplineMeshComp->DetachFromParent(true);
+							SplineMeshComp->InvalidateLightingCache();
+							SplineMeshComp->Rename(NULL, Landscape);
+							SplineMeshComp->AttachTo(Landscape->SplineComponent, NAME_None, EAttachLocation::KeepWorldPosition);
 						}
 
 						Segment->UpdateSplinePoints();
diff --git a/Engine/Source/Editor/LandscapeEditor/Private/LandscapeEditorModule.cpp b/Engine/Source/Editor/LandscapeEditor/Private/LandscapeEditorModule.cpp
--- a/Engine/Source/Editor/LandscapeEditor/Private/LandscapeEditorModule.cpp
+++ b/Engine/Source/Editor/LandscapeEditor/Private/LandscapeEditorModule.cpp
@@ -79,18 +79,18 @@ public:
 	{
 		struct Local
 		{
-			static void BuildLandscapeVisualizersMenu(FMenuBuilder& MenuBuilder)
+			static void BuildLandscapeVisualizersMenu(FMenuBuilder& InMenuBuilder)
 			{
 				const FLandscapeEditorCommands& LandscapeActions = FLandscapeEditorCommands::Get();
 
-				MenuBuilder.BeginSection("LandscapeVisualizers", LOCTEXT("LandscapeHeader", "Landscape Visualizers"));
+				InMenuBuilder.BeginSection("LandscapeVisualizers", LOCTEXT("LandscapeHeader", "Landscape Visualizers"));
 				{
-					MenuBuilder.AddMenuEntry(LandscapeActions.ViewModeNormal,       NAME_None, LOCTEXT("LandscapeViewModeNormal", "Normal"));
-					MenuBuilder.AddMenuEntry(LandscapeActions.ViewModeLOD,          NAME_None, LOCTEXT("LandscapeViewModeLOD", "LOD"));
-					MenuBuilder.AddMenuEntry(LandscapeActions.ViewModeLayerDensity, NAME_None, LOCTEXT("LandscapeViewModeLayerDensity", "Layer Density"));
-					MenuBuilder.AddMenuEntry(LandscapeActions.ViewModeLayerDebug,   NAME_None, LOCTEXT("LandscapeViewModeLayerDebug", "Layer Debug"));
+					InMenuBuilder.AddMenuEntry(LandscapeActions.ViewModeNormal, NAME_None, LOCTEXT("LandscapeViewModeNormal", "Normal"));
+					InMenuBuilder.AddMenuEntry(LandscapeActions.ViewModeLOD, NAME_None, LOCTEXT("LandscapeViewModeLOD", "LOD"));
+					InMenuBuilder.AddMenuEntry(LandscapeActions.ViewModeLayerDensity, NAME_None, LOCTEXT("LandscapeViewModeLayerDensity", "Layer Density"));
+					InMenuBuilder.AddMenuEntry(LandscapeActions.ViewModeLayerDebug, NAME_None, LOCTEXT("LandscapeViewModeLayerDebug", "Layer Debug"));
 				}
-				MenuBuilder.EndSection();
+				InMenuBuilder.EndSection();
 			}
 		};
 		MenuBuilder.AddSubMenu(LOCTEXT("LandscapeSubMenu", "Visualizers"), LOCTEXT("LandscapeSubMenu_ToolTip", "Select a Landscape visualiser"), FNewMenuDelegate::CreateStatic(&Local::BuildLandscapeVisualizersMenu));
diff --git a/Engine/Source/Editor/LevelEditor/Private/LevelEditor.cpp b/Engine/Source/Editor/LevelEditor/Private/LevelEditor.cpp
--- a/Engine/Source/Editor/LevelEditor/Private/LevelEditor.cpp
+++ b/Engine/Source/Editor/LevelEditor/Private/LevelEditor.cpp
@@ -263,9 +263,9 @@ void FLevelEditorModule::AttachSequencer(TSharedPtr<SWidget> SequencerWidget, TS
 	{
 		struct Local
 		{
-			static void OnSequencerClosed( TSharedRef<SDockTab> DockTab, TWeakPtr<IAssetEditorInstance> SequencerAssetEditor )
+			static void OnSequencerClosed( TSharedRef<SDockTab> DockTab, TWeakPtr<IAssetEditorInstance> InSequencerAssetEditor )
 			{
-				SequencerAssetEditor.Pin()->CloseWindow();
+				InSequencerAssetEditor.Pin()->CloseWindow();
 			}
 		};
 		
diff --git a/Engine/Source/Editor/LevelEditor/Private/SLevelEditor.cpp b/Engine/Source/Editor/LevelEditor/Private/SLevelEditor.cpp
--- a/Engine/Source/Editor/LevelEditor/Private/SLevelEditor.cpp
+++ b/Engine/Source/Editor/LevelEditor/Private/SLevelEditor.cpp
@@ -629,11 +629,11 @@ TSharedRef<SDockTab> SLevelEditor::SpawnLevelEditorTab( const FSpawnTabArgs& Arg
 			InitOptions.DefaultMenuExtender = MakeShareable(new FExtender);
 			InitOptions.DefaultMenuExtender->AddMenuExtension(
 				"FolderSection", EExtensionHook::Before, GetLevelEditorActions(),
-				FMenuExtensionDelegate::CreateStatic([](FMenuBuilder& MenuBuilder, TWeakPtr<SLevelEditor> WeakLevelEditor){
+				FMenuExtensionDelegate::CreateStatic([](FMenuBuilder& MenuBuilder, TWeakPtr<SLevelEditor> InWeakLevelEditor){
 					// Only extend the menu if we have actors selected
 					if (GEditor->GetSelectedActors()->Num())
 					{
-						FLevelEditorContextMenu::FillMenu(MenuBuilder, WeakLevelEditor, LevelEditorMenuContext::NonViewport, TSharedPtr<FExtender>());
+						FLevelEditorContextMenu::FillMenu(MenuBuilder, InWeakLevelEditor, LevelEditorMenuContext::NonViewport, TSharedPtr<FExtender>());
 					}
 				}, WeakLevelEditor)
 			);
diff --git a/Engine/Source/Editor/TranslationEditor/Private/TranslationDataManager.cpp b/Engine/Source/Editor/TranslationEditor/Private/TranslationDataManager.cpp
--- a/Engine/Source/Editor/TranslationEditor/Private/TranslationDataManager.cpp
+++ b/Engine/Source/Editor/TranslationEditor/Private/TranslationDataManager.cpp
@@ -337,7 +337,7 @@ bool FTranslationDataManager::WriteJSONToTextFile(TSharedRef<FJsonObject>& Outpu
 	return CheckoutAndSaveWasSuccessful;
 }
 
-void FTranslationDataManager::GetHistoryForTranslationUnits( TArray<UTranslationUnit*>& TranslationUnits, const FString& ManifestFilePath )
+void FTranslationDataManager::GetHistoryForTranslationUnits( TArray<UTranslationUnit*>& TranslationUnits, const FString& InManifestFilePath )
 {
 	GWarn->BeginSlowTask(LOCTEXT("LoadingSourceControlHistory", "Loading Translation History from Source Control..."), true);
 
@@ -345,13 +345,13 @@ void FTranslationDataManager::GetHistoryForTranslationUnits( TArray<UTranslation
 	ISourceControlProvider& SourceControlProvider = ISourceControlModule::Get().GetProvider();
 	TSharedRef<FUpdateStatus, ESPMode::ThreadSafe> UpdateStatusOperation = ISourceControlOperation::Create<FUpdateStatus>();
 	UpdateStatusOperation->SetUpdateHistory( true );
-	ECommandResult::Type Result = SourceControlProvider.Execute(UpdateStatusOperation, ManifestFilePath);
+	ECommandResult::Type Result = SourceControlProvider.Execute(UpdateStatusOperation, InManifestFilePath);
 	bool bGetHistoryFromSourceControlSucceeded = Result == ECommandResult::Succeeded;
 
 	// Now we can get information about the file's history from the source control state, retrieve that
 	TArray<FString> Files;
 	TArray< TSharedRef<ISourceControlState, ESPMode::ThreadSafe> > States;
-	Files.Add(ManifestFilePath);
+	Files.Add(InManifestFilePath);
 	Result = SourceControlProvider.GetState( Files,  States, EStateCacheUsage::ForceUpdate );
 	bGetHistoryFromSourceControlSucceeded = bGetHistoryFromSourceControlSucceeded && (Result == ECommandResult::Succeeded);
 	FSourceControlStatePtr SourceControlState;
@@ -372,7 +372,7 @@ void FTranslationDataManager::GetHistoryForTranslationUnits( TArray<UTranslation
 			TSharedPtr<ISourceControlRevision, ESPMode::ThreadSafe> Revision = SourceControlState->GetHistoryItem(HistoryItemIndex);
 			if(Revision.IsValid())
 			{
-				FString ManifestFullPath = FPaths::ConvertRelativePathToFull(ManifestFilePath);
+				FString ManifestFullPath = FPaths::ConvertRelativePathToFull(InManifestFilePath);
 				FString EngineFullPath = FPaths::ConvertRelativePathToFull(FPaths::EngineContentDir());
 
 				bool IsEngineManifest = false;
@@ -394,7 +394,7 @@ void FTranslationDataManager::GetHistoryForTranslationUnits( TArray<UTranslation
 					SavedDir = FPaths::GameSavedDir();
 				}
 
-				FString TempFileName = SavedDir / "CachedTranslationHistory" / "UE4-Manifest-" + ProjectName + "-" + FPaths::GetBaseFilename(ManifestFilePath) + "-Rev-" + FString::FromInt(Revision->GetRevisionNumber());
+				FString TempFileName = SavedDir / "CachedTranslationHistory" / "UE4-Manifest-" + ProjectName + "-" + FPaths::GetBaseFilename(InManifestFilePath) + "-Rev-" + FString::FromInt(Revision->GetRevisionNumber());
 
 				
 				if (!FPaths::FileExists(TempFileName))	// Don't bother syncing again if we already have this manifest version cached locally
@@ -451,7 +451,7 @@ void FTranslationDataManager::GetHistoryForTranslationUnits( TArray<UTranslation
 				else // OldManifestPtr.IsValid() is false
 				{
 					FFormatNamedArguments Arguments;
-					Arguments.Add( TEXT("ManifestFilePath"), FText::FromString(ManifestFilePath) );
+					Arguments.Add(TEXT("ManifestFilePath"), FText::FromString(InManifestFilePath));
 					Arguments.Add( TEXT("ManifestRevisionNumber"), FText::AsNumber(Revision->GetRevisionNumber()) );
 					FMessageLog TranslationEditorMessageLog("TranslationEditor");
 					TranslationEditorMessageLog.Warning(FText::Format(LOCTEXT("PreviousManifestCorrupt", "Previous revision {ManifestRevisionNumber} of {ManifestFilePath} failed to load correctly. Ignoring."), Arguments));
@@ -463,7 +463,7 @@ void FTranslationDataManager::GetHistoryForTranslationUnits( TArray<UTranslation
 	else // (bGetHistoryFromSourceControlSucceeded && SourceControlState.IsValid()) is false
 	{
 		FFormatNamedArguments Arguments;
-		Arguments.Add( TEXT("ManifestFilePath"), FText::FromString(ManifestFilePath) );
+		Arguments.Add(TEXT("ManifestFilePath"), FText::FromString(InManifestFilePath));
 		FMessageLog TranslationEditorMessageLog("SourceControl");
 		TranslationEditorMessageLog.Warning(FText::Format(LOCTEXT("SourceControlStateQueryFailed", "Failed to query source control state of file {ManifestFilePath}."), Arguments));
 		TranslationEditorMessageLog.Notify(LOCTEXT("RetrieveTranslationHistoryFailed", "Unable to Retrieve Translation History from Source Control!"));
diff --git a/Engine/Source/Editor/UnrealEd/Private/CookOnTheFlyServer.cpp b/Engine/Source/Editor/UnrealEd/Private/CookOnTheFlyServer.cpp
--- a/Engine/Source/Editor/UnrealEd/Private/CookOnTheFlyServer.cpp
+++ b/Engine/Source/Editor/UnrealEd/Private/CookOnTheFlyServer.cpp
@@ -893,14 +893,13 @@ void UCookOnTheFlyServer::HandleNetworkFileServerFileRequest( const FString& Fil
 
 	if (!bIsCookable)
 	{
-		TArray<FFilePlatformRequest> FileRequests;
-		UnsolicitedCookedPackages.DequeueAll(FileRequests);
-		for ( int Index=  0; Index < FileRequests.Num(); ++Index)
+		TArray<FFilePlatformRequest> UnsolicitedFileRequests;
+		UnsolicitedCookedPackages.DequeueAll(UnsolicitedFileRequests);
+		for (const FFilePlatformRequest &Request : UnsolicitedFileRequests)
 		{
-			const FFilePlatformRequest &Request = FileRequests[Index];
 			if ( Request.Platformname == PlatformFname )
 			{
-				FString StandardFilename = FileRequests[Index].Filename;
+				FString StandardFilename = Request.Filename;
 				FPaths::MakeStandardFilename( StandardFilename );
 				UnsolicitedFiles.Add( StandardFilename );
 			}
@@ -925,11 +924,10 @@ void UCookOnTheFlyServer::HandleNetworkFileServerFileRequest( const FString& Fil
 
 	UE_LOG( LogCookOnTheFly, Display, TEXT("Cook complete %s"), *FileRequest.Filename)
 
-	TArray<FFilePlatformRequest> FileRequests;
-	UnsolicitedCookedPackages.DequeueAll(FileRequests);
-	for ( int Index=  0; Index < FileRequests.Num(); ++Index)
+	TArray<FFilePlatformRequest> UnsolicitedFileRequests;
+	UnsolicitedCookedPackages.DequeueAll(UnsolicitedFileRequests);
+	for (const FFilePlatformRequest &Request : UnsolicitedFileRequests)
 	{
-		const FFilePlatformRequest &Request = FileRequests[Index];
 		if ( Request.Filename == FileRequest.Filename )
 		{
 			// don't do anything we don't want this guy
diff --git a/Engine/Source/Editor/UnrealEd/Private/EditorViewportClient.cpp b/Engine/Source/Editor/UnrealEd/Private/EditorViewportClient.cpp
--- a/Engine/Source/Editor/UnrealEd/Private/EditorViewportClient.cpp
+++ b/Engine/Source/Editor/UnrealEd/Private/EditorViewportClient.cpp
@@ -81,10 +81,10 @@ FViewportCameraTransform::FViewportCameraTransform()
 	: TransitionCurve( new FCurveSequence( 0.0f, FocusConstants::TransitionTime, ECurveEaseFunction::CubicOut ) )
 	, ViewLocation( FVector::ZeroVector )
 	, ViewRotation( FRotator::ZeroRotator )
-	, DesiredLocation(FVector::ZeroVector)
+	, DesiredLocation( FVector::ZeroVector )
 	, LookAt( FVector::ZeroVector )
-	, StartLocation(FVector::ZeroVector)
-	, OrthoZoom(DEFAULT_ORTHOZOOM)
+	, StartLocation( FVector::ZeroVector )
+	, OrthoZoom( DEFAULT_ORTHOZOOM )
 {}
 
 void FViewportCameraTransform::SetLocation( const FVector& Position )
@@ -203,7 +203,7 @@ FEditorViewportClient::FEditorViewportClient(FPreviewScene* InPreviewScene)
 	, bNeedsLinkedRedraw(false)
 	, bNeedsInvalidateHitProxy(false)
 	, bUsingOrbitCamera(false)
-	, bDisableInput(false)
+	, bDisableInput( false )
 	, bDrawAxes(true)
 	, bSetListenerPosition(false)
 	, LandscapeLODOverride(-1)
@@ -600,11 +600,11 @@ FSceneView* FEditorViewportClient::CalcSceneView(FSceneViewFamily* ViewFamily)
 			OrthoHeight = (ControllingActorViewInfo.OrthoWidth / 2.0f) * YScale;
 
 			const float NearViewPlane = ViewInitOptions.ViewMatrix.GetOrigin().Z;
-			const float FarPlane = NearViewPlane - 2.0f*WORLD_MAX;
+			const float FarViewPlane = NearViewPlane - 2.0f*WORLD_MAX;
 
-			const float InverseRange = 1.0f / (FarPlane - NearViewPlane);
+			const float InverseRange = 1.0f / (FarViewPlane - NearViewPlane);
 			ZScale = -2.0f * InverseRange;
-			ZOffset = -(FarPlane + NearViewPlane) * InverseRange;
+			ZOffset = -(FarViewPlane + NearViewPlane) * InverseRange;
 		}
 		else
 		{
@@ -1501,10 +1501,6 @@ void FEditorViewportClient::UpdateMouseDelta()
 								}
 								MoveViewportCamera( Drag, Rot );
 
-								const bool LeftMouseButtonDown = Viewport->KeyState(EKeys::LeftMouseButton) ? true : false;
-								const bool MiddleMouseButtonDown = Viewport->KeyState(EKeys::MiddleMouseButton) ? true : false;
-								const bool RightMouseButtonDown = Viewport->KeyState(EKeys::RightMouseButton) ? true : false;
-
 								if ( IsPerspective() && LeftMouseButtonDown && !MiddleMouseButtonDown && !RightMouseButtonDown )
 								{
 									FEditorViewportStats::Using(FEditorViewportStats::CAT_PERSPECTIVE_MOUSE_DOLLY);
@@ -3685,9 +3681,9 @@ void FEditorViewportClient::OnJoystickPlugged(const uint32 InControllerID, const
 }
 
 
-void FEditorViewportClient::MouseEnter(FViewport* Viewport,int32 x, int32 y)
+void FEditorViewportClient::MouseEnter(FViewport* InViewport,int32 x, int32 y)
 {
-	MouseMove(Viewport, x, y);
+	MouseMove(InViewport, x, y);
 }
 
 void FEditorViewportClient::MouseMove(FViewport* InViewport,int32 x, int32 y)
@@ -3697,13 +3693,13 @@ void FEditorViewportClient::MouseMove(FViewport* InViewport,int32 x, int32 y)
 	CurrentMousePos = FIntPoint(x, y);
 }
 
-void FEditorViewportClient::MouseLeave( FViewport* Viewport )
+void FEditorViewportClient::MouseLeave(FViewport* InViewport)
 {
 	check(IsInGameThread());
 
 	CurrentMousePos = FIntPoint(-1, -1);
 
-	FCommonViewportClient::MouseLeave(Viewport);
+	FCommonViewportClient::MouseLeave(InViewport);
 }
 
 void FEditorViewportClient::CapturedMouseMove( FViewport* InViewport, int32 InMouseX, int32 InMouseY )
diff --git a/Engine/Source/Editor/UnrealEd/Private/Fbx/FbxSkeletalMeshImport.cpp b/Engine/Source/Editor/UnrealEd/Private/Fbx/FbxSkeletalMeshImport.cpp
--- a/Engine/Source/Editor/UnrealEd/Private/Fbx/FbxSkeletalMeshImport.cpp
+++ b/Engine/Source/Editor/UnrealEd/Private/Fbx/FbxSkeletalMeshImport.cpp
@@ -738,10 +738,10 @@ bool UnFbx::FFbxImporter::ImportBone(TArray<FbxNode*>& NodeArray, FSkeletalMeshI
 					else
 					{
 						// first try to find parent who is null group and see if you can try test it again
-						FbxNode * Parent = Current->GetParent();
-						while (Parent)
+						FbxNode * ParentNode = Current->GetParent();
+						while (ParentNode)
 						{
-							FbxNodeAttribute* Attr = Parent->GetNodeAttribute();
+							FbxNodeAttribute* Attr = ParentNode->GetNodeAttribute();
 							if ( Attr && Attr->GetAttributeType() == FbxNodeAttribute::eNull )
 							{
 								// found it 
@@ -749,10 +749,10 @@ bool UnFbx::FFbxImporter::ImportBone(TArray<FbxNode*>& NodeArray, FSkeletalMeshI
 							}
 
 							// find next parent
-							Parent = Parent->GetParent();
+							ParentNode = ParentNode->GetParent();
 						}
 
-						if ( Parent && CurrentPose->IsValidBindPose(Parent) )
+						if (ParentNode && CurrentPose->IsValidBindPose(ParentNode))
 						{
 							PoseArray.Add(CurrentPose);
 							UE_LOG(LogFbx, Warning, TEXT("Valid bind pose for Pose (%s) - %s"), *PoseName, *FString(Current->GetName()));
@@ -1096,7 +1096,7 @@ USkeletalMesh* UnFbx::FFbxImporter::ImportSkeletalMesh(UObject* InParent, TArray
 	TArray<FbxSurfaceMaterial*> FbxMaterials;
 	for( int32 NodeIndex = 0; NodeIndex < NodeArray.Num(); ++NodeIndex )
 	{
-		FbxNode* Node = NodeArray[NodeIndex];
+		Node = NodeArray[NodeIndex];
 
 		int32 MaterialCount = Node->GetMaterialCount();
 
diff --git a/Engine/Source/Editor/WorldBrowser/Private/Tiles/WorldTileModel.cpp b/Engine/Source/Editor/WorldBrowser/Private/Tiles/WorldTileModel.cpp
--- a/Engine/Source/Editor/WorldBrowser/Private/Tiles/WorldTileModel.cpp
+++ b/Engine/Source/Editor/WorldBrowser/Private/Tiles/WorldTileModel.cpp
@@ -881,28 +881,28 @@ ALandscapeProxy* FWorldTileModel::ImportLandscape(const FLandscapeImportSettings
 		return nullptr;
 	}
 	
-	ALandscapeProxy*	Landscape;
+	ALandscapeProxy*	LandscapeProxy;
 	FGuid				LandscapeGuid = Settings.LandscapeGuid;
 	
 	if (LandscapeGuid.IsValid())
 	{
-		Landscape = Cast<UWorld>(LoadedLevel->GetOuter())->SpawnActor<ALandscapeProxy>();
+		LandscapeProxy = Cast<UWorld>(LoadedLevel->GetOuter())->SpawnActor<ALandscapeProxy>();
 	}
 	else
 	{
-		Landscape = Cast<UWorld>(LoadedLevel->GetOuter())->SpawnActor<ALandscape>();
+		LandscapeProxy = Cast<UWorld>(LoadedLevel->GetOuter())->SpawnActor<ALandscape>();
 		LandscapeGuid = FGuid::NewGuid();
 	}
 	
-	Landscape->SetActorTransform(Settings.LandscapeTransform);
+	LandscapeProxy->SetActorTransform(Settings.LandscapeTransform);
 	
 	if (Settings.LandscapeMaterial)
 	{
-		Landscape->LandscapeMaterial = Settings.LandscapeMaterial;
+		LandscapeProxy->LandscapeMaterial = Settings.LandscapeMaterial;
 	}
 	
 	// Create landscape components	
-	Landscape->Import(
+	LandscapeProxy->Import(
 		LandscapeGuid, 
 		Settings.SizeX, 
 		Settings.SizeY, 
@@ -913,7 +913,7 @@ ALandscapeProxy* FWorldTileModel::ImportLandscape(const FLandscapeImportSettings
 		*Settings.HeightmapFilename, 
 		Settings.ImportLayers);
 
-	return Landscape;
+	return LandscapeProxy;
 }
 
 #undef LOCTEXT_NAMESPACE
diff --git a/Engine/Source/Runtime/Engine/Private/Landscape/LandscapeCollision.cpp b/Engine/Source/Runtime/Engine/Private/Landscape/LandscapeCollision.cpp
--- a/Engine/Source/Runtime/Engine/Private/Landscape/LandscapeCollision.cpp
+++ b/Engine/Source/Runtime/Engine/Private/Landscape/LandscapeCollision.cpp
@@ -266,8 +266,10 @@ void ULandscapeHeightfieldCollisionComponent::CreateCollisionObject()
 				HeightfieldRef = GSharedHeightfieldRefs.Add(HeightfieldGuid, new FPhysXHeightfieldRef(HeightfieldGuid));
 
 				// Create heightfield shape
-				FPhysXInputStream HeightFieldStream(CookedCollisionData.GetData(), CookedCollisionData.Num());
-				HeightfieldRef->RBHeightfield = GPhysXSDK->createHeightField(HeightFieldStream);
+				{
+					FPhysXInputStream HeightFieldStream(CookedCollisionData.GetData(), CookedCollisionData.Num());
+					HeightfieldRef->RBHeightfield = GPhysXSDK->createHeightField(HeightFieldStream);
+				}
 
 				for (UPhysicalMaterial* PhysicalMaterial : CookedPhysicalMaterials)
 				{
diff --git a/Engine/Source/Runtime/Engine/Private/Landscape/LandscapeSplines.cpp b/Engine/Source/Runtime/Engine/Private/Landscape/LandscapeSplines.cpp
--- a/Engine/Source/Runtime/Engine/Private/Landscape/LandscapeSplines.cpp
+++ b/Engine/Source/Runtime/Engine/Private/Landscape/LandscapeSplines.cpp
@@ -1451,13 +1451,13 @@ void ULandscapeSplineSegment::UpdateSplinePoints(bool bUpdateCollision)
 			const UStaticMesh* const Mesh = MeshComponent->StaticMesh;
 			const FBoxSphereBounds MeshBounds = Mesh->GetBounds();
 
-			const float T = MeshSettings[i].T * Rescale;
+			const float RescaledT = MeshSettings[i].T * Rescale;
 			const FLandscapeSplineMeshEntry* MeshEntry = MeshSettings[i].MeshEntry;
 			const ESplineMeshAxis::Type SideAxis = CrossAxis(MeshEntry->ForwardAxis, MeshEntry->UpAxis);
 
 			const float TEnd = MeshSettings[i + 1].T * Rescale;
 
-			const float CosInterp = 0.5f - 0.5f * FMath::Cos(T * PI);
+			const float CosInterp = 0.5f - 0.5f * FMath::Cos(RescaledT * PI);
 			const float Width = FMath::Lerp(StartWidth, EndWidth, CosInterp);
 			const bool bDoOrientationRoll = (MeshEntry->ForwardAxis == ESplineMeshAxis::X && MeshEntry->UpAxis == ESplineMeshAxis::Y) ||
 			                                (MeshEntry->ForwardAxis == ESplineMeshAxis::Y && MeshEntry->UpAxis == ESplineMeshAxis::Z) ||
@@ -1502,8 +1502,8 @@ void ULandscapeSplineSegment::UpdateSplinePoints(bool bUpdateCollision)
 			Offset *= Scale2D;
 			Offset = RotateVec2D(Offset, -Roll);
 
-			MeshComponent->SplineParams.StartPos = SplineInfo.Eval(T, FVector::ZeroVector);
-			MeshComponent->SplineParams.StartTangent = SplineInfo.EvalDerivative(T, FVector::ZeroVector) * (TEnd - T);
+			MeshComponent->SplineParams.StartPos = SplineInfo.Eval(RescaledT, FVector::ZeroVector);
+			MeshComponent->SplineParams.StartTangent = SplineInfo.EvalDerivative(RescaledT, FVector::ZeroVector) * (TEnd - RescaledT);
 			MeshComponent->SplineParams.StartScale = Scale2D;
 			MeshComponent->SplineParams.StartRoll = Roll;
 			MeshComponent->SplineParams.StartOffset = Offset;
@@ -1551,7 +1551,7 @@ void ULandscapeSplineSegment::UpdateSplinePoints(bool bUpdateCollision)
 			OffsetEnd = RotateVec2D(OffsetEnd, -RollEnd);
 
 			MeshComponent->SplineParams.EndPos = SplineInfo.Eval(TEnd, FVector::ZeroVector);
-			MeshComponent->SplineParams.EndTangent = SplineInfo.EvalDerivative(TEnd, FVector::ZeroVector) * (TEnd - T);
+			MeshComponent->SplineParams.EndTangent = SplineInfo.EvalDerivative(TEnd, FVector::ZeroVector) * (TEnd - RescaledT);
 			MeshComponent->SplineParams.EndScale = Scale2DEnd;
 			MeshComponent->SplineParams.EndRoll = RollEnd;
 			MeshComponent->SplineParams.EndOffset = OffsetEnd;
diff --git a/Engine/Source/Runtime/Engine/Private/Materials/Material.cpp b/Engine/Source/Runtime/Engine/Private/Materials/Material.cpp
--- a/Engine/Source/Runtime/Engine/Private/Materials/Material.cpp
+++ b/Engine/Source/Runtime/Engine/Private/Materials/Material.cpp
@@ -1979,7 +1979,6 @@ void UMaterial::PostLoad()
 		ITargetPlatformManagerModule* TPM = GetTargetPlatformManager();
 		if (TPM && (TPM->RestrictFormatsToRuntimeOnly() == false))
 		{
-			ITargetPlatformManagerModule* TPM = GetTargetPlatformManager();
 			TArray<ITargetPlatform*> Platforms = TPM->GetActiveTargetPlatforms();
 			// Cache for all the shader formats that the cooking target requires
 			for (int32 FormatIndex = 0; FormatIndex < Platforms.Num(); FormatIndex++)
@@ -2046,10 +2045,10 @@ void UMaterial::BeginCacheForCookedPlatformData( const ITargetPlatform *TargetPl
 			// Cache for all the shader formats that the cooking target requires
 			for (int32 FormatIndex = 0; FormatIndex < DesiredShaderFormats.Num(); FormatIndex++)
 			{
-				const EShaderPlatform TargetPlatform = ShaderFormatToLegacyShaderPlatform(DesiredShaderFormats[FormatIndex]);
+				const EShaderPlatform LegacyShaderPlatform = ShaderFormatToLegacyShaderPlatform(DesiredShaderFormats[FormatIndex]);
 
 				// Begin caching shaders for the target platform and store the material resource being compiled into CachedMaterialResourcesForCooking
-				CacheResourceShadersForCooking(TargetPlatform, *CachedMaterialResourcesForPlatform);
+				CacheResourceShadersForCooking(LegacyShaderPlatform, *CachedMaterialResourcesForPlatform);
 			}
 		}
 	}
diff --git a/Engine/Source/Runtime/Engine/Private/Materials/MaterialExpressions.cpp b/Engine/Source/Runtime/Engine/Private/Materials/MaterialExpressions.cpp
--- a/Engine/Source/Runtime/Engine/Private/Materials/MaterialExpressions.cpp
+++ b/Engine/Source/Runtime/Engine/Private/Materials/MaterialExpressions.cpp
@@ -3922,12 +3922,12 @@ bool UMaterialExpressionQualitySwitch::IsInputConnectionRequired(int32 InputInde
 bool UMaterialExpressionQualitySwitch::IsResultMaterialAttributes(int32 OutputIndex)
 {
 	check(OutputIndex == 0);
-	TArray<FExpressionInput*> Inputs = GetInputs();
+	TArray<FExpressionInput*> ExpressionInputs = GetInputs();
 
-	for (int32 Index = 0; Index < Inputs.Num(); ++Index)
+	for (FExpressionInput* ExpressionInput : ExpressionInputs)
 	{
 		// If there is a loop anywhere in this expression's inputs then we can't risk checking them
-		if (Inputs[Index]->Expression && !Inputs[Index]->Expression->ContainsInputLoop() && Inputs[Index]->Expression->IsResultMaterialAttributes(Inputs[Index]->OutputIndex))
+		if (ExpressionInput->Expression && !ExpressionInput->Expression->ContainsInputLoop() && ExpressionInput->Expression->IsResultMaterialAttributes(ExpressionInput->OutputIndex))
 		{
 			return true;
 		}
@@ -4025,12 +4025,12 @@ bool UMaterialExpressionFeatureLevelSwitch::IsInputConnectionRequired(int32 Inpu
 bool UMaterialExpressionFeatureLevelSwitch::IsResultMaterialAttributes(int32 OutputIndex)
 {
 	check(OutputIndex == 0);
-	TArray<FExpressionInput*> Inputs = GetInputs();
+	TArray<FExpressionInput*> ExpressionInputs = GetInputs();
 
-	for (int32 Index = 0; Index < Inputs.Num(); ++Index)
+	for (FExpressionInput* ExpressionInput : ExpressionInputs)
 	{
 		// If there is a loop anywhere in this expression's inputs then we can't risk checking them
-		if (Inputs[Index]->Expression && !Inputs[Index]->Expression->ContainsInputLoop() && Inputs[Index]->Expression->IsResultMaterialAttributes(Inputs[Index]->OutputIndex))
+		if (ExpressionInput->Expression && !ExpressionInput->Expression->ContainsInputLoop() && ExpressionInput->Expression->IsResultMaterialAttributes(ExpressionInput->OutputIndex))
 		{
 			return true;
 		}
diff --git a/Engine/Source/Runtime/Engine/Private/Slate/DebugCanvas.cpp b/Engine/Source/Runtime/Engine/Private/Slate/DebugCanvas.cpp
--- a/Engine/Source/Runtime/Engine/Private/Slate/DebugCanvas.cpp
+++ b/Engine/Source/Runtime/Engine/Private/Slate/DebugCanvas.cpp
@@ -114,16 +114,16 @@ FCanvas* FDebugCanvasDrawer::GetGameThreadDebugCanvas()
 }
 
 
-void FDebugCanvasDrawer::BeginRenderingCanvas( const FIntRect& InCanvasRect )
+void FDebugCanvasDrawer::BeginRenderingCanvas( const FIntRect& CanvasRect )
 {
-	if( InCanvasRect.Size().X > 0 && InCanvasRect.Size().Y > 0 )
+	if( CanvasRect.Size().X > 0 && CanvasRect.Size().Y > 0 )
 	{
 		ENQUEUE_UNIQUE_RENDER_COMMAND_THREEPARAMETER
 		(
 			BeginRenderingDebugCanvas,
 			FDebugCanvasDrawer*, CanvasDrawer, this, 
 			FCanvasProxy*, CanvasToRender, GameThreadCanvas,
-			FIntRect, CanvasRect, InCanvasRect,
+			FIntRect, CanvasRect, CanvasRect,
 			{
 				// Delete the old rendering thread canvas
 				if( CanvasDrawer->GetRenderThreadCanvas() && CanvasToRender != NULL )
diff --git a/Engine/Source/Runtime/Engine/Private/UserInterface/Canvas.cpp b/Engine/Source/Runtime/Engine/Private/UserInterface/Canvas.cpp
--- a/Engine/Source/Runtime/Engine/Private/UserInterface/Canvas.cpp
+++ b/Engine/Source/Runtime/Engine/Private/UserInterface/Canvas.cpp
@@ -528,9 +528,9 @@ void FCanvas::Flush(bool bForce)
 			// disable depth test & writes
 			RHISetDepthStencilState(TStaticDepthStencilState<false,CF_Always>::GetRHI());
 
-			const FIntRect& ViewRect = Parameters.ViewRect;
+			const FIntRect& ViewportRect = Parameters.ViewRect;
 			// set viewport to RT size
-			RHISetViewport( ViewRect.Min.X,ViewRect.Min.Y, 0.0f, ViewRect.Max.X, ViewRect.Max.Y, 1.0f );	
+			RHISetViewport(ViewportRect.Min.X, ViewportRect.Min.Y, 0.0f, ViewportRect.Max.X, ViewportRect.Max.Y, 1.0f);
 		});
 	}
 
diff --git a/Engine/Source/Runtime/SkillSystem/Public/SkillSystemLog.h b/Engine/Source/Runtime/SkillSystem/Public/SkillSystemLog.h
--- a/Engine/Source/Runtime/SkillSystem/Public/SkillSystemLog.h
+++ b/Engine/Source/Runtime/SkillSystem/Public/SkillSystemLog.h
@@ -17,7 +17,9 @@ SKILLSYSTEM_API DECLARE_LOG_CATEGORY_EXTERN(LogSkillSystem, Warning, All);
 	UE_LOG(LogSkillSystem, Verbosity, TEXT("%s"), *Str); \
 }
 
-#define SKILL_LOG_SCOPE( Format, ... ) SkillSystemLogScope LogScope##__LINE__( FString::Printf(Format, ##__VA_ARGS__));
+#define SKILL_LOG_TOKENPASTE_INNER(x,y) x##y
+#define SKILL_LOG_TOKENPASTE(x,y) SKILL_LOG_TOKENPASTE_INNER(x,y)
+#define SKILL_LOG_SCOPE( Format, ... ) SkillSystemLogScope SKILL_LOG_TOKENPASTE(LogScope,__LINE__)( FString::Printf(Format, ##__VA_ARGS__));
 
 struct SkillSystemLogScope
 {
