Translation Editor: Fixing some incorrect/missing toolbar icons. Still just placeholder.

diff --git a/Engine/Source/Editor/EditorStyle/Private/SlateEditorStyle.cpp b/Engine/Source/Editor/EditorStyle/Private/SlateEditorStyle.cpp
--- a/Engine/Source/Editor/EditorStyle/Private/SlateEditorStyle.cpp
+++ b/Engine/Source/Editor/EditorStyle/Private/SlateEditorStyle.cpp
@@ -5504,7 +5504,8 @@ void FSlateEditorStyle::FStyle::SetupUMGEditorStyles()
 
 void FSlateEditorStyle::FStyle::SetupTranslationEditorStyles()
 {
-	Set("TranslationEditor.Import", new IMAGE_BRUSH("Icons/icon_Import_40x", Icon40x40));
+	Set("TranslationEditor.Export", new IMAGE_BRUSH("Icons/icon_Import_40x", Icon40x40));
+	Set("TranslationEditor.Import", new IMAGE_BRUSH("Icons/icon_Samples_Download_40x", Icon40x40));
 }
 
 #undef IMAGE_BRUSH
diff --git a/Engine/Source/Editor/TranslationEditor/Private/TranslationEditorMenu.cpp b/Engine/Source/Editor/TranslationEditor/Private/TranslationEditorMenu.cpp
--- a/Engine/Source/Editor/TranslationEditor/Private/TranslationEditorMenu.cpp
+++ b/Engine/Source/Editor/TranslationEditor/Private/TranslationEditorMenu.cpp
@@ -70,7 +70,7 @@ void FTranslationEditorMenu::SetupTranslationEditorToolbar( TSharedPtr< FExtende
 			ToolbarBuilder.AddToolBarButton(
 				FTranslationEditorCommands::Get().PreviewAllTranslationsInEditor, "PreviewTranslationsInEditor", TAttribute<FText>(), TAttribute<FText>(), FSlateIcon(FEditorStyle::GetStyleSetName(), "AssetEditor.ReimportAsset")); 
 			ToolbarBuilder.AddToolBarButton(
-				FTranslationEditorCommands::Get().ExportToPortableObjectFormat, "ExportToPortableObjectFormat", TAttribute<FText>(), TAttribute<FText>(), FSlateIcon(FEditorStyle::GetStyleSetName(), "FontEditor.Export"));
+				FTranslationEditorCommands::Get().ExportToPortableObjectFormat, "ExportToPortableObjectFormat", TAttribute<FText>(), TAttribute<FText>(), FSlateIcon(FEditorStyle::GetStyleSetName(), "TranslationEditor.Export"));
 			ToolbarBuilder.AddToolBarButton(
 				FTranslationEditorCommands::Get().ImportFromPortableObjectFormat, "ImportFromPortableObjectFormat", TAttribute<FText>(), TAttribute<FText>(), FSlateIcon(FEditorStyle::GetStyleSetName(), "TranslationEditor.Import"));
 			ToolbarBuilder.AddToolBarButton(
