Fix to comply with CL 2276830 "Consistent localization-related commandlet path-handling with regards to relative paths. From now on, please specify paths relative to to the project's directory or the Engine directory."

Changed SourcePath for the temp InternationalizationExport.ini to be ./Content/... instead of ../../../Engine/Content/...

diff --git a/Engine/Source/Editor/TranslationEditor/Private/TranslationEditor.cpp b/Engine/Source/Editor/TranslationEditor/Private/TranslationEditor.cpp
--- a/Engine/Source/Editor/TranslationEditor/Private/TranslationEditor.cpp
+++ b/Engine/Source/Editor/TranslationEditor/Private/TranslationEditor.cpp
@@ -1111,11 +1111,35 @@ void FTranslationEditor::ExportToPortableObjectFormat_Execute()
 		// Write translation data first to ensure all changes are exported
 		DataManager->WriteTranslationData();
 
+		// Source path needs to be relative to Engine or Game directory
+		FString ManifestFullPath = FPaths::ConvertRelativePathToFull(ManifestFilePath);
+		FString EngineFullPath = FPaths::ConvertRelativePathToFull(FPaths::EngineContentDir());
+		FString FolderRelativeManifestPath = ManifestFilePath;
+		bool IsEngineManifest = false;
+		if (ManifestFullPath.StartsWith(EngineFullPath))
+		{
+			IsEngineManifest = true;
+		}
+		if (IsEngineManifest)
+		{
+			FPaths::MakePathRelativeTo(FolderRelativeManifestPath, *FPaths::EngineDir());
+		}
+		else
+		{
+			FPaths::MakePathRelativeTo(FolderRelativeManifestPath, *FPaths::GameDir());
+		}
+
+		// Paths like "Content/Localization/Editor/" are not considered relative for some reason, so add "./" if neccessary
+		if (!FPaths::IsRelative(FolderRelativeManifestPath))
+		{
+			FolderRelativeManifestPath = FString(TEXT("./")) + FolderRelativeManifestPath;
+		}
+
 		UInternationalizationExportSettings* ExportSettings = NewObject<UInternationalizationExportSettings>();
 		ExportSettings->CulturesToGenerate.Empty();
 		ExportSettings->CulturesToGenerate.Add(FPaths::GetBaseFilename(FPaths::GetPath(ArchiveFilePath)));
 		ExportSettings->CommandletClass = "InternationalizationExport";
-		ExportSettings->SourcePath = FPaths::GetPath(ManifestFilePath);
+		ExportSettings->SourcePath = FPaths::GetPath(FolderRelativeManifestPath);
 		ExportSettings->ManifestName = FPaths::GetBaseFilename(ManifestFilePath) + ".manifest";
 		ExportSettings->ArchiveName = FPaths::GetBaseFilename(ManifestFilePath) + ".archive";
 		ExportSettings->bExportLoc = true;
