Fix to comply with CL 2276830 "Consistent localization-related commandlet path-handling with regards to relative paths. From now on, please specify paths relative to to the project's directory or the Engine directory."

diff --git a/Engine/Source/Editor/TranslationEditor/Private/TranslationEditor.cpp b/Engine/Source/Editor/TranslationEditor/Private/TranslationEditor.cpp
--- a/Engine/Source/Editor/TranslationEditor/Private/TranslationEditor.cpp
+++ b/Engine/Source/Editor/TranslationEditor/Private/TranslationEditor.cpp
@@ -1259,11 +1259,35 @@ void FTranslationEditor::ImportFromPortableObjectFormat_Execute()
 
 	if (bOpened)
 	{
+		// Source path needs to be relative to Engine or Game directory
+		FString ManifestFullPath = FPaths::ConvertRelativePathToFull(ManifestFilePath);
+		FString EngineFullPath = FPaths::ConvertRelativePathToFull(FPaths::EngineContentDir());
+		FString FolderRelativeManifestPath = ManifestFilePath;
+		bool IsEngineManifest = false;
+		if (ManifestFullPath.StartsWith(EngineFullPath))
+		{
+			IsEngineManifest = true;
+		}
+		if (IsEngineManifest)
+		{
+			FPaths::MakePathRelativeTo(FolderRelativeManifestPath, *FPaths::EngineDir());
+		}
+		else
+		{
+			FPaths::MakePathRelativeTo(FolderRelativeManifestPath, *FPaths::GameDir());
+		}
+
+		// Paths like "Content/Localization/Editor/" are not considered relative for some reason, so add "./" if neccessary
+		if (!FPaths::IsRelative(FolderRelativeManifestPath))
+		{
+			FolderRelativeManifestPath = FString(TEXT("./")) + FolderRelativeManifestPath;
+		}
+
 		UInternationalizationExportSettings* ImportSettings = NewObject<UInternationalizationExportSettings>();
 		ImportSettings->CulturesToGenerate.Empty();
 		ImportSettings->CulturesToGenerate.Add(FPaths::GetBaseFilename(FPaths::GetPath(ArchiveFilePath)));
 		ImportSettings->CommandletClass = "InternationalizationExport";
-		ImportSettings->DestinationPath = FPaths::GetPath(ManifestFilePath);
+		ImportSettings->DestinationPath = FPaths::GetPath(FolderRelativeManifestPath);
 		ImportSettings->ManifestName = FPaths::GetBaseFilename(ManifestFilePath) + ".manifest";
 		ImportSettings->ArchiveName = FPaths::GetBaseFilename(ManifestFilePath) + ".archive";
 		ImportSettings->bExportLoc = false;
