TranslationEditor: also search Namespace, key and context when using the search tab

diff --git a/Engine/Source/Editor/TranslationEditor/Private/TranslationDataManager.cpp b/Engine/Source/Editor/TranslationEditor/Private/TranslationDataManager.cpp
--- a/Engine/Source/Editor/TranslationEditor/Private/TranslationDataManager.cpp
+++ b/Engine/Source/Editor/TranslationEditor/Private/TranslationDataManager.cpp
@@ -513,10 +513,24 @@ void FTranslationDataManager::PopulateSearchResultsUsingFilter(const FString& Se
 	{
 		if (TranslationUnit != NULL)
 		{
+			bool bAdded = false;
 			if (TranslationUnit->Source.Contains(SearchFilter) ||
-				TranslationUnit->Translation.Contains(SearchFilter))
+				TranslationUnit->Translation.Contains(SearchFilter) || 
+				TranslationUnit->Namespace.Contains(SearchFilter))
 			{
 				SearchResults.Add(TranslationUnit);
+				bAdded = true;
+			}
+			
+			for (FTranslationContextInfo CurrentContext : TranslationUnit->Contexts)
+			{
+				if (!bAdded && 
+					(CurrentContext.Context.Contains(SearchFilter) ||
+					CurrentContext.Key.Contains(SearchFilter)))
+				{
+					SearchResults.Add(TranslationUnit);
+					break;
+				}
 			}
 		}
 	}
