Translation Editor: Remove all escaping and unescaping of special characters.  This is now all handled by the JSON serializer.

diff --git a/Engine/Source/Editor/TranslationEditor/Private/TranslationDataManager.cpp b/Engine/Source/Editor/TranslationEditor/Private/TranslationDataManager.cpp
--- a/Engine/Source/Editor/TranslationEditor/Private/TranslationDataManager.cpp
+++ b/Engine/Source/Editor/TranslationEditor/Private/TranslationDataManager.cpp
@@ -54,7 +54,7 @@ FTranslationDataManager::FTranslationDataManager( const FString& InManifestFileP
 				// We want Undo/Redo support
 				TranslationUnit->SetFlags(RF_Transactional);
 				TranslationUnit->HasBeenReviewed = false;
-				TranslationUnit->Source = ManifestEntry->Source.Text.ReplaceEscapedCharWithChar();
+				TranslationUnit->Source = ManifestEntry->Source.Text;
 				TranslationUnit->Namespace = ManifestEntry->Namespace;
 
 				for(auto ContextIter( ManifestEntry->Contexts.CreateConstIterator() ); ContextIter; ++ContextIter)
@@ -178,7 +178,7 @@ bool FTranslationDataManager::WriteTranslationData(bool bForceWrite /*= false*/)
 		{
 			const FLocItem SearchSource(TranslationUnit->Source);
 			FString OldTranslation = Archive->FindEntryBySource(TranslationUnit->Namespace, SearchSource, NULL)->Translation.Text;
-			FString TranslationToWrite = TranslationUnit->Translation.ReplaceCharWithEscapedChar();
+			FString TranslationToWrite = TranslationUnit->Translation;
 			if (!TranslationToWrite.Equals(OldTranslation))
 			{
 				Archive->SetTranslation(TranslationUnit->Namespace, TranslationUnit->Source, TranslationToWrite, NULL);
@@ -193,7 +193,7 @@ bool FTranslationDataManager::WriteTranslationData(bool bForceWrite /*= false*/)
 		{
 			const FLocItem SearchSource(TranslationUnit->Source);
 			FString OldTranslation = Archive->FindEntryBySource(TranslationUnit->Namespace, SearchSource, NULL)->Translation.Text;
-			FString TranslationToWrite = TranslationUnit->Translation.ReplaceCharWithEscapedChar();
+			FString TranslationToWrite = TranslationUnit->Translation;
 			if (TranslationUnit->HasBeenReviewed && !TranslationToWrite.Equals(OldTranslation))
 			{
 				Archive->SetTranslation(TranslationUnit->Namespace, TranslationUnit->Source, TranslationToWrite, NULL);
@@ -208,7 +208,7 @@ bool FTranslationDataManager::WriteTranslationData(bool bForceWrite /*= false*/)
 		{
 			const FLocItem SearchSource(TranslationUnit->Source);
 			FString OldTranslation = Archive->FindEntryBySource(TranslationUnit->Namespace, SearchSource, NULL)->Translation.Text;
-			FString TranslationToWrite = TranslationUnit->Translation.ReplaceCharWithEscapedChar();
+			FString TranslationToWrite = TranslationUnit->Translation;
 			if (!TranslationToWrite.Equals(OldTranslation))
 			{
 				Archive->SetTranslation(TranslationUnit->Namespace, TranslationUnit->Source, TranslationToWrite, NULL);
@@ -432,14 +432,14 @@ void FTranslationDataManager::GetHistoryForTranslationUnits( TArray<UTranslation
 								}
 
 								// Always add first instance of this string, and then add any versions that changed since
-								if (ContextInfo.Changes.Num() == 0 || !OldManifestEntryPtr->Source.Text.ReplaceEscapedCharWithChar().Equals(PreviousSourceText))
+								if (ContextInfo.Changes.Num() == 0 || !OldManifestEntryPtr->Source.Text.Equals(PreviousSourceText))
 								{
 									TSharedPtr< FArchiveEntry > OldArchiveEntry = ArchivePtr->FindEntryBySource(OldManifestEntryPtr->Namespace, OldManifestEntryPtr->Source, NULL);
 									if (OldArchiveEntry.IsValid())
 									{
 										FTranslationChange Change;
-										Change.Source = OldManifestEntryPtr->Source.Text.ReplaceEscapedCharWithChar();
-										Change.Translation = OldArchiveEntry->Translation.Text.ReplaceEscapedCharWithChar();
+										Change.Source = OldManifestEntryPtr->Source.Text;
+										Change.Translation = OldArchiveEntry->Translation.Text;
 										Change.DateAndTime = Revision->GetDate();
 										Change.Version = FString::FromInt(Revision->GetRevisionNumber());
 										ContextInfo.Changes.Insert(Change, 0);
@@ -578,9 +578,9 @@ void FTranslationDataManager::LoadFromArchive(TArray<UTranslationUnit*>& InTrans
 				{
 					const FString PreviousTranslation = TranslationUnit->Translation;
 					TranslationUnit->Translation = ""; // Reset to null string
-					const FString UnescapedTranslatedString = ArchiveEntry->Translation.Text.ReplaceEscapedCharWithChar();
+					const FString TranslatedString = ArchiveEntry->Translation.Text;
 
-					if (UnescapedTranslatedString.IsEmpty())
+					if (TranslatedString.IsEmpty())
 					{
 						bool bHasTranslationHistory = false;
 						int32 MostRecentNonNullTranslationIndex = -1;
@@ -619,7 +619,7 @@ void FTranslationDataManager::LoadFromArchive(TArray<UTranslationUnit*>& InTrans
 					}
 					else
 					{
-						TranslationUnit->Translation = UnescapedTranslatedString;
+						TranslationUnit->Translation = TranslatedString;
 						TranslationUnit->HasBeenReviewed = true;
 						Complete.Add(TranslationUnit);
 					}
