Translation Editor: Initial work for a "Translation picker." In current state when enabled, hover over FText in the editor and it will show you the information about that text and any tooltips.  Shows Key, Namespace, Source, Translation and Table Name.  Works for many strings, some still fail to be looked up in the Localization tables, will continue investigating.

diff --git a/Engine/Source/Editor/Documentation/Private/SDocumentationToolTip.h b/Engine/Source/Editor/Documentation/Private/SDocumentationToolTip.h
deleted file mode 100644
diff --git a/Engine/Source/Editor/Documentation/Public/SDocumentationToolTip.h b/Engine/Source/Editor/Documentation/Public/SDocumentationToolTip.h
new file mode 100644
--- /dev/null
+++ b/Engine/Source/Editor/Documentation/Public/SDocumentationToolTip.h
@@ -0,0 +1,85 @@
+// Copyright 1998-2014 Epic Games, Inc. All Rights Reserved.
+#pragma once
+#include "IDocumentationPage.h"
+
+class SDocumentationToolTip : public SCompoundWidget
+{
+
+public:
+
+	SLATE_BEGIN_ARGS( SDocumentationToolTip )
+		: _Text()
+		, _Style( TEXT("Documentation.SDocumentationTooltip") )
+		, _ColorAndOpacity( FLinearColor( 1.0f, 1.0f, 1.0f, 1.0f ) )
+		, _Content()
+		{}
+
+		/** The text displayed in this tool tip */
+		SLATE_ATTRIBUTE( FText, Text )
+
+		/** The text style to use for this tool tip */
+		SLATE_ARGUMENT( FName, Style )
+		
+		/** Font color and opacity */
+		SLATE_ATTRIBUTE( FSlateColor, ColorAndOpacity )
+
+		/**  */
+		SLATE_ARGUMENT( FString, DocumentationLink )
+
+		/**  */
+		SLATE_ARGUMENT( FString, ExcerptName )
+
+		/** Arbitrary content to be displayed in the tool tip; overrides any text that may be set. */
+		SLATE_DEFAULT_SLOT( FArguments, Content )
+
+	SLATE_END_ARGS()
+
+	/**
+	 * Construct this widget
+	 *
+	 * @param	InArgs	The declaration data for this widget
+	 */
+	void Construct( const FArguments& InArgs );
+
+	void Tick( const FGeometry& AllottedGeometry, const double InCurrentTime, const float InDeltaTime );
+
+	bool IsInteractive() const;
+
+	virtual const FText& GetTextTooltip() const
+	{
+		return TextContent.Get();
+	}
+
+private:
+
+	void ConstructSimpleTipContent();
+
+	void ConstructFullTipContent();
+
+	FReply ReloadDocumentation();
+
+	void CreateExcerpt( FString FileSource, FString ExcerptName );
+
+private:
+
+	/** Text block widget */
+	TAttribute< FText > TextContent;
+	TSharedPtr< SWidget > OverrideContent;
+	FTextBlockStyle StyleInfo;
+	TAttribute< FSlateColor > ColorAndOpacity;
+
+	/** The link to the documentation */
+	FString DocumentationLink;
+	FString ExcerptName;
+
+	/** Content widget */
+	TSharedPtr< SBorder > WidgetContent;
+
+	TSharedPtr< SWidget > SimpleTipContent;
+	bool IsDisplayingDocumentationLink;
+
+	TSharedPtr< SWidget > FullTipContent;
+
+	TSharedPtr< IDocumentationPage > DocumentationPage;
+	bool IsShowingFullTip;
+};
diff --git a/Engine/Source/Editor/EditorStyle/Private/SlateEditorStyle.cpp b/Engine/Source/Editor/EditorStyle/Private/SlateEditorStyle.cpp
--- a/Engine/Source/Editor/EditorStyle/Private/SlateEditorStyle.cpp
+++ b/Engine/Source/Editor/EditorStyle/Private/SlateEditorStyle.cpp
@@ -6059,7 +6059,10 @@ void FSlateEditorStyle::FStyle::SetupUMGEditorStyles()
 void FSlateEditorStyle::FStyle::SetupTranslationEditorStyles()
 {
 	Set("TranslationEditor.Export", new IMAGE_BRUSH("Icons/icon_Import_40x", Icon40x40));
+	Set("TranslationEditor.PreviewInEditor", new IMAGE_BRUSH("Icons/icon_levels_visible_40x", Icon40x40));
 	Set("TranslationEditor.Import", new IMAGE_BRUSH("Icons/icon_Samples_Download_40x", Icon40x40));
+	Set("TranslationEditor.Search", new IMAGE_BRUSH("Icons/icon_Blueprint_Find_40px", Icon40x40));
+	Set("TranslationEditor.TranslationPicker", new IMAGE_BRUSH("Icons/icon_StaticMeshEd_VertColor_40x", Icon40x40));
 }
 
 #undef IMAGE_BRUSH
diff --git a/Engine/Source/Editor/TranslationEditor/Private/TranslationEditorMenu.cpp b/Engine/Source/Editor/TranslationEditor/Private/TranslationEditorMenu.cpp
--- a/Engine/Source/Editor/TranslationEditor/Private/TranslationEditorMenu.cpp
+++ b/Engine/Source/Editor/TranslationEditor/Private/TranslationEditorMenu.cpp
@@ -6,8 +6,9 @@
 #include "GraphEditorActions.h"
 #include "Editor/PropertyEditor/Public/PropertyEditing.h"
 #include "PropertyCustomizationHelpers.h"
+#include "TranslationPickerWidget.h"
 
-#define LOCTEXT_NAMESPACE "KismetToolbar"
+#define LOCTEXT_NAMESPACE "TranslationEditorToolbar"
 
 void FTranslationEditorMenu::FillTranslationMenu( FMenuBuilder& MenuBuilder/*, FTranslationEditor& TranslationEditor*/ )
 {
@@ -68,13 +69,14 @@ void FTranslationEditorMenu::SetupTranslationEditorToolbar( TSharedPtr< FExtende
 			ToolbarBuilder.AddToolBarButton(
 				FTranslationEditorCommands::Get().SaveTranslations, "SaveTranslations", TAttribute<FText>(), TAttribute<FText>(), FSlateIcon(FEditorStyle::GetStyleSetName(), "AssetEditor.SaveAsset"));
 			ToolbarBuilder.AddToolBarButton(
-				FTranslationEditorCommands::Get().PreviewAllTranslationsInEditor, "PreviewTranslationsInEditor", TAttribute<FText>(), TAttribute<FText>(), FSlateIcon(FEditorStyle::GetStyleSetName(), "AssetEditor.ReimportAsset")); 
+				FTranslationEditorCommands::Get().PreviewAllTranslationsInEditor, "PreviewTranslationsInEditor", TAttribute<FText>(), TAttribute<FText>(), FSlateIcon(FEditorStyle::GetStyleSetName(), "TranslationEditor.PreviewInEditor")); 
 			ToolbarBuilder.AddToolBarButton(
 				FTranslationEditorCommands::Get().ExportToPortableObjectFormat, "ExportToPortableObjectFormat", TAttribute<FText>(), TAttribute<FText>(), FSlateIcon(FEditorStyle::GetStyleSetName(), "TranslationEditor.Export"));
 			ToolbarBuilder.AddToolBarButton(
 				FTranslationEditorCommands::Get().ImportFromPortableObjectFormat, "ImportFromPortableObjectFormat", TAttribute<FText>(), TAttribute<FText>(), FSlateIcon(FEditorStyle::GetStyleSetName(), "TranslationEditor.Import"));
 			ToolbarBuilder.AddToolBarButton(
-				FTranslationEditorCommands::Get().OpenSearchTab, "OpenSearchTab", TAttribute<FText>(), TAttribute<FText>(), FSlateIcon(FEditorStyle::GetStyleSetName(), "BlueprintEditor.FindInBlueprint"));
+				FTranslationEditorCommands::Get().OpenSearchTab, "OpenSearchTab", TAttribute<FText>(), TAttribute<FText>(), FSlateIcon(FEditorStyle::GetStyleSetName(), "TranslationEditor.Search"));
+			ToolbarBuilder.AddWidget(SNew(STranslationWidgetPicker));
 		}
 	};
 
diff --git a/Engine/Source/Editor/TranslationEditor/Private/TranslationPickerFloatingWindow.cpp b/Engine/Source/Editor/TranslationEditor/Private/TranslationPickerFloatingWindow.cpp
new file mode 100644
--- /dev/null
+++ b/Engine/Source/Editor/TranslationEditor/Private/TranslationPickerFloatingWindow.cpp
@@ -0,0 +1,152 @@
+
+#include "TranslationEditorPrivatePCH.h"
+#include "TranslationPickerFloatingWindow.h"
+#include "Editor/Documentation/Public/SDocumentationToolTip.h"
+
+#define LOCTEXT_NAMESPACE "TranslationPicker"
+
+void STranslationPickerFloatingWindow::Construct(const FArguments& InArgs)
+{
+	ParentWindow = InArgs._ParentWindow;
+
+	ChildSlot
+		[
+			SNew(SToolTip)
+			.Text(this, &STranslationPickerFloatingWindow::GetPickerStatusText)
+		];
+}
+
+void STranslationPickerFloatingWindow::Tick(const FGeometry& AllottedGeometry, const double InCurrentTime, const float InDeltaTime)
+{
+	TranslationInfoPreviewText = FText::GetEmpty();
+
+	FWidgetPath Path = FSlateApplication::Get().LocateWindowUnderMouse(FSlateApplication::Get().GetCursorPos(), FSlateApplication::Get().GetInteractiveTopLevelWindows(), true);
+
+	// Search everything under the cursor for any FText we know how to parse
+	for (int32 PathIndex = Path.Widgets.Num() - 1; PathIndex >= 0; PathIndex--)
+	{
+		// General Widget case
+		TSharedRef<SWidget> PathWidget = Path.Widgets[PathIndex].Widget;
+		FText TextBlockDescription = GetTextDescription(PathWidget);
+
+		if (!TextBlockDescription.IsEmpty())
+		{
+			TranslationInfoPreviewText = FText::Format(LOCTEXT("TextBlockLocalizationDescription", "{0}\n\n\n\n\nTextBlockInfo: {1}"),
+				TranslationInfoPreviewText, TextBlockDescription);
+		}
+
+		// Tooltip case
+		TSharedPtr<IToolTip> Tooltip = PathWidget->GetToolTip();
+		FText TooltipDescription = FText::GetEmpty();
+
+		if (Tooltip.IsValid() && !Tooltip->IsEmpty())
+		{
+			TooltipDescription = GetTextDescription(Tooltip->AsWidget());
+			if (!TooltipDescription.IsEmpty())
+			{
+				TranslationInfoPreviewText = FText::Format(LOCTEXT("TooltipLocalizationDescription", "{0}\n\n\n\n\nTooltip Info:{1}"),
+					TranslationInfoPreviewText, TooltipDescription);
+			}
+		}
+		
+	}
+
+	// kind of a hack, but we need to maintain keyboard focus otherwise we wont get our keypress to 'pick'
+	FSlateApplication::Get().SetKeyboardFocus(SharedThis(this), EKeyboardFocusCause::SetDirectly);
+	if (ParentWindow.IsValid())
+	{
+		// also kind of a hack, but this is the only way at the moment to get a 'cursor decorator' without using the drag-drop code path
+		ParentWindow.Pin()->MoveWindowTo(FSlateApplication::Get().GetCursorPos() + FSlateApplication::Get().GetCursorSize());
+	}
+}
+
+FText STranslationPickerFloatingWindow::GetTextDescription(TSharedRef<SWidget> Widget)
+{
+	FText OutText = FText::GetEmpty();
+
+	// Parse STextBlocks
+	STextBlock* TextBlock = (STextBlock*)&Widget.Get();
+	if ((Widget->GetTypeAsString() == "STextBlock" && TextBlock))
+	{
+		FText Text = TextBlock->GetText();
+		OutText = FormatFTextInfo(Text);
+	}
+
+	// Parse SToolTips
+	else if (Widget->GetTypeAsString() == "SToolTip")
+	{
+		SToolTip* ToolTip = ((SToolTip*)&Widget.Get());
+		if (ToolTip != nullptr)
+		{
+			OutText = GetTextDescription(ToolTip->GetContentWidget());
+			if (OutText.IsEmpty())
+			{
+				OutText = FormatFTextInfo(ToolTip->GetTextTooltip());
+			}
+		}
+	}
+
+	// Parse SDocumentationToolTips
+	else if (Widget->GetTypeAsString() == "SDocumentationToolTip")
+	{
+		SDocumentationToolTip* DocumentationToolTip = (SDocumentationToolTip*)&Widget.Get();
+		if (DocumentationToolTip != nullptr)
+		{
+			OutText = FormatFTextInfo(DocumentationToolTip->GetTextTooltip());
+		}
+	}
+	return OutText;
+}
+
+FText STranslationPickerFloatingWindow::FormatFTextInfo(FText TextToFormat)
+{
+	FText OutText = FText::GetEmpty();
+
+	if (FTextInspector::ShouldGatherForLocalization(TextToFormat))
+	{
+		const FString* TextNamespace = FTextInspector::GetNamespace(TextToFormat);
+		const FString* TextKey = FTextInspector::GetKey(TextToFormat);
+		const FString* TextSource = FTextInspector::GetSourceString(TextToFormat);
+		const FString& TextTranslation = FTextInspector::GetDisplayString(TextToFormat);
+
+		TSharedPtr<FString, ESPMode::ThreadSafe> TextTableName = TSharedPtr<FString, ESPMode::ThreadSafe>(nullptr);
+		if (TextNamespace && TextKey)
+		{
+			TextTableName = FTextLocalizationManager::Get().GetTableName(*TextNamespace, *TextKey);
+		}
+
+		FText Namespace = TextNamespace != nullptr ? FText::FromString(*TextNamespace) : FText::GetEmpty();
+		FText Key = TextKey != nullptr ? FText::FromString(*TextKey) : FText::GetEmpty();
+		FText Source = TextSource != nullptr ? FText::FromString(*TextSource) : FText::GetEmpty();
+		FText TableName = TextTableName.IsValid() ? FText::FromString(*(TextTableName.Get())) : FText::GetEmpty();
+		FText Translation = FText::FromString(TextTranslation);
+
+		OutText = FText::Format(LOCTEXT("LocalizationStringDescription", "\n\n\n\n\n\nText Namespace: {0}\nText Key: {1}\nText Source: {2}\nTable Name: {3}"), Namespace, Key, Source, TableName);
+		OutText = FText::Format(LOCTEXT("LocalizationStringDescription2", "{0}\nTranslation: {1}"), OutText, Translation);
+	}
+
+	return OutText; 
+}
+
+FReply STranslationPickerFloatingWindow::OnKeyDown(const FGeometry& MyGeometry, const FKeyboardEvent& InKeyboardEvent)
+{
+	if (InKeyboardEvent.GetKey() == EKeys::Escape)
+	{
+		//TODO: pop up window for in-place translation
+
+		TranslationInfoPreviewText = FText::GetEmpty();
+
+		if (ParentWindow.IsValid())
+		{
+			FSlateApplication::Get().RequestDestroyWindow(ParentWindow.Pin().ToSharedRef());
+			ParentWindow.Reset();
+		}
+
+		return FReply::Handled();
+	}
+
+	return FReply::Unhandled();
+}
+
+
+#undef LOCTEXT_NAMESPACE
\ No newline at end of file
diff --git a/Engine/Source/Editor/TranslationEditor/Private/TranslationPickerFloatingWindow.h b/Engine/Source/Editor/TranslationEditor/Private/TranslationPickerFloatingWindow.h
new file mode 100644
--- /dev/null
+++ b/Engine/Source/Editor/TranslationEditor/Private/TranslationPickerFloatingWindow.h
@@ -0,0 +1,51 @@
+#pragma once
+#include "SCompoundWidget.h"
+
+#define LOCTEXT_NAMESPACE "TranslationPicker"
+
+/** Translation picker window to show details of FText(s) under cursor, and allow in-place translation */
+class STranslationPickerFloatingWindow : public SCompoundWidget
+{
+	SLATE_BEGIN_ARGS(STranslationPickerFloatingWindow) {}
+
+	SLATE_ARGUMENT(TWeakPtr<SWindow>, ParentWindow)
+
+	SLATE_END_ARGS()
+
+	void Construct(const FArguments& InArgs);
+
+	FText GetTranslationInfoPreviewText()
+	{
+		return TranslationInfoPreviewText;
+	}
+
+private:
+	virtual void Tick(const FGeometry& AllottedGeometry, const double InCurrentTime, const float InDeltaTime) override;
+
+	/** Pull the FText reference out of an SWidget */
+	FText GetTextDescription(TSharedRef<SWidget> Widget);
+
+	/** Format an FText for display */
+	FText FormatFTextInfo(FText TextToFormat);
+
+	FText GetPickerStatusText() const
+	{
+		return FText::Format(LOCTEXT("TootipHint", "{0}\n\n(Esc to pick)"), TranslationInfoPreviewText);
+	}
+
+	virtual FReply OnKeyDown(const FGeometry& MyGeometry, const FKeyboardEvent& InKeyboardEvent) override;
+
+	/** We need to support keyboard focus to process the 'Esc' key */
+	virtual bool SupportsKeyboardFocus() const override
+	{
+		return true;
+	}
+
+	/** Handle to the window that contains this widget */
+	TWeakPtr<SWindow> ParentWindow;
+
+	/** The widget name we are picking */
+	FText TranslationInfoPreviewText;
+};
+
+#undef LOCTEXT_NAMESPACE
\ No newline at end of file
diff --git a/Engine/Source/Editor/TranslationEditor/Private/TranslationPickerWidget.cpp b/Engine/Source/Editor/TranslationEditor/Private/TranslationPickerWidget.cpp
new file mode 100644
--- /dev/null
+++ b/Engine/Source/Editor/TranslationEditor/Private/TranslationPickerWidget.cpp
@@ -0,0 +1,148 @@
+// Copyright 1998-2014 Epic Games, Inc. All Rights Reserved.
+
+#include "TranslationEditorPrivatePCH.h"
+#include "TranslationPickerWidget.h"
+
+#define LOCTEXT_NAMESPACE "TranslationPicker"
+
+STranslationWidgetPicker::~STranslationWidgetPicker()
+{
+	// kill the picker window as well if this widget is going away - that way we dont get dangling refs to the property
+	if (PickerWindow.IsValid() && FSlateApplication::IsInitialized())
+	{
+		FSlateApplication::Get().RequestDestroyWindow(PickerWindow.Pin().ToSharedRef());
+		PickerWindow.Reset();
+		PickerWindowWidget.Reset();
+	}
+}
+
+void STranslationWidgetPicker::Construct(const FArguments& InArgs)
+{
+	// Mimicking a toolbar button look
+
+	// Icon for the picker widget button
+	TSharedRef< SWidget > IconWidget =
+		SNew(SImage)
+		.Image(FEditorStyle::GetBrush("TranslationEditor.TranslationPicker"));
+
+	// Style settings
+	FName StyleSet = FEditorStyle::GetStyleSetName();
+	FName StyleName = "Toolbar";
+
+	FText ToolTip = LOCTEXT("TranslationPickerTooltip", "Open the Translation Picker");
+
+	// Create the content for our button
+	TSharedRef< SWidget > ButtonContent =
+
+		SNew(SHorizontalBox)
+
+		+ SHorizontalBox::Slot()
+		.FillWidth(1)
+		.VAlign(VAlign_Center)
+		[
+			SNew(SVerticalBox)
+
+			// Icon image
+			+ SVerticalBox::Slot()
+			.AutoHeight()
+			.HAlign(HAlign_Center)	// Center the icon horizontally, so that large labels don't stretch out the artwork
+			[
+				IconWidget
+			]
+
+			// Label text
+			+ SVerticalBox::Slot().AutoHeight()
+				.HAlign(HAlign_Center)	// Center the label text horizontally
+				[
+					SNew(STextBlock)
+					.Text(LOCTEXT("TranslationPicker", "Translation Picker"))
+					.TextStyle(FEditorStyle::Get(), FName("ToolBar.Label"))	// Smaller font for tool tip labels
+					.ShadowOffset(FVector2D::UnitVector)
+				]
+		];
+
+	FName CheckboxStyle = ISlateStyle::Join(StyleName, ".SToolBarButtonBlock.CheckBox.Padding");
+
+	ChildSlot
+		[
+			// Create a check box
+			SNew(SCheckBox)
+
+			// Use the tool bar style for this check box
+			.Style(FEditorStyle::Get(), "ToolBar.ToggleButton")
+
+			// User will have set the focusable attribute for the block, honor it
+			.IsFocusable(false)
+
+			// Pass along the block's tool-tip string
+			.ToolTip(SNew(SToolTip).Text(ToolTip))
+			[
+				ButtonContent
+			]
+
+			// Bind the button's "on checked" event to our object's method for this
+			.OnCheckStateChanged(this, &STranslationWidgetPicker::OnCheckStateChanged)
+
+
+				// Bind the check box's "checked" state to our user interface action
+				.IsChecked(this, &STranslationWidgetPicker::OnIsChecked)
+
+				.Padding(FEditorStyle::Get().GetMargin(CheckboxStyle))
+		];
+}
+
+FReply STranslationWidgetPicker::OnClicked()
+{
+	// Not picking previously, launch a picker window
+	if (!PickerWindow.IsValid())
+	{
+		TSharedRef<SWindow> NewWindow = SWindow::MakeCursorDecorator();
+		NewWindow->MoveWindowTo(FSlateApplication::Get().GetCursorPos());
+		PickerWindow = NewWindow;
+
+		NewWindow->SetContent(
+			SAssignNew(PickerWindowWidget, STranslationPickerFloatingWindow)
+			.ParentWindow(NewWindow)
+			);
+
+		TSharedPtr<SWindow> RootWindow = FGlobalTabmanager::Get()->GetRootWindow();
+		check(RootWindow.IsValid());
+		FSlateApplication::Get().AddWindowAsNativeChild(NewWindow, RootWindow.ToSharedRef());
+	}
+	// Was already picking previously, destroy picker window
+	else
+	{
+		FSlateApplication::Get().RequestDestroyWindow(PickerWindow.Pin().ToSharedRef());
+		PickerWindow.Pin().Reset();
+	}
+
+	return FReply::Handled();
+}
+
+bool STranslationWidgetPicker::IsPicking() const
+{
+	if (PickerWindowWidget.IsValid())
+	{
+		return true;
+	}
+
+	return false;
+}
+
+ESlateCheckBoxState::Type STranslationWidgetPicker::OnIsChecked() const
+{
+	return IsPicking() ? ESlateCheckBoxState::Checked : ESlateCheckBoxState::Unchecked;
+}
+
+void STranslationWidgetPicker::OnCheckStateChanged(const ESlateCheckBoxState::Type NewCheckedState)
+{
+	OnClicked();
+}
+
+
+
+
+#undef LOCTEXT_NAMESPACE
+
+
+
diff --git a/Engine/Source/Editor/TranslationEditor/Private/TranslationPickerWidget.h b/Engine/Source/Editor/TranslationEditor/Private/TranslationPickerWidget.h
new file mode 100644
--- /dev/null
+++ b/Engine/Source/Editor/TranslationEditor/Private/TranslationPickerWidget.h
@@ -0,0 +1,45 @@
+// Copyright 1998-2014 Epic Games, Inc. All Rights Reserved.
+#pragma once
+
+#include "SDocumentationToolTip.h"
+#include "TranslationPickerFloatingWindow.h"
+
+#define LOCTEXT_NAMESPACE "TranslationPicker"
+
+/** Widget used to launch a 'picking' session */
+class STranslationWidgetPicker : public SCompoundWidget
+{
+public:
+	SLATE_BEGIN_ARGS(STranslationWidgetPicker) {}
+
+	SLATE_END_ARGS()
+
+	~STranslationWidgetPicker();
+
+	void Construct(const FArguments& InArgs);
+
+private:
+	FReply OnClicked();
+
+	bool IsPicking() const;
+
+	/** Picker window widget */
+	TWeakPtr<STranslationPickerFloatingWindow> PickerWindowWidget;
+
+	/** Picker window */
+	TWeakPtr<SWindow> PickerWindow;
+
+	/**
+	* Called by slate to determine if this button should appear checked
+	*
+	* @return ESlateCheckBoxState::Checked if it should be checked, ESlateCheckBoxState::Unchecked if not.
+	*/
+	ESlateCheckBoxState::Type OnIsChecked() const;
+
+	/**
+	* Called by Slate when this tool bar check box button is toggled
+	*/
+	void OnCheckStateChanged(const ESlateCheckBoxState::Type NewCheckedState);
+};
+
+#undef LOCTEXT_NAMESPACE
\ No newline at end of file
diff --git a/Engine/Source/Editor/TranslationEditor/TranslationEditor.Build.cs b/Engine/Source/Editor/TranslationEditor/TranslationEditor.Build.cs
--- a/Engine/Source/Editor/TranslationEditor/TranslationEditor.Build.cs
+++ b/Engine/Source/Editor/TranslationEditor/TranslationEditor.Build.cs
@@ -32,6 +32,7 @@ public class TranslationEditor : ModuleRules
                 "GraphEditor",
 				"SourceControl",
                 "MessageLog",
+                "Documentation",
 			}
 		);
 
diff --git a/Engine/Source/Runtime/Core/Private/Internationalization/TextLocalizationManager.cpp b/Engine/Source/Runtime/Core/Private/Internationalization/TextLocalizationManager.cpp
--- a/Engine/Source/Runtime/Core/Private/Internationalization/TextLocalizationManager.cpp
+++ b/Engine/Source/Runtime/Core/Private/Internationalization/TextLocalizationManager.cpp
@@ -500,6 +500,25 @@ TSharedPtr<FString, ESPMode::ThreadSafe> FTextLocalizationManager::FindString( c
 	return NULL;
 }
 
+TSharedPtr<FString, ESPMode::ThreadSafe> FTextLocalizationManager::GetTableName(const FString& Namespace, const FString& Key)
+{
+	FScopeLock ScopeLock(&SynchronizationObject);
+
+	// Find namespace's key table.
+	const FTextLookupTable::FKeyTable* LiveKeyTable = LiveTable.NamespaceTable.Find(Namespace);
+
+	// Find key table's entry.
+	const FStringEntry* LiveEntry = LiveKeyTable ? LiveKeyTable->Find(Key) : NULL;
+
+	if (LiveEntry != nullptr)
+	{
+		return MakeShareable(new FString((LiveEntry->TableName)));
+	}
+
+	return NULL;
+}
+
+
 TSharedRef<FString, ESPMode::ThreadSafe> FTextLocalizationManager::GetString(const FString& Namespace, const FString& Key, const FString* const SourceString)
 {
 	FScopeLock ScopeLock( &SynchronizationObject );
diff --git a/Engine/Source/Runtime/Core/Public/Internationalization/TextLocalizationManager.h b/Engine/Source/Runtime/Core/Public/Internationalization/TextLocalizationManager.h
--- a/Engine/Source/Runtime/Core/Public/Internationalization/TextLocalizationManager.h
+++ b/Engine/Source/Runtime/Core/Public/Internationalization/TextLocalizationManager.h
@@ -91,6 +91,9 @@ public:
 	 */
 	void CORE_API FindKeyNamespaceFromDisplayString(TSharedRef<FString, ESPMode::ThreadSafe> InDisplayString, TSharedPtr<FString, ESPMode::ThreadSafe>& OutNamespace, TSharedPtr<FString, ESPMode::ThreadSafe>& OutKey);
 
+	/** Get table name for given Namespace/key */
+	TSharedPtr<FString, ESPMode::ThreadSafe> CORE_API GetTableName(const FString& Namespace, const FString& Key);
+
 	void CORE_API RegenerateResources( const FString& ConfigFilePath, IInternationalizationArchiveSerializer& ArchiveSerializer, IInternationalizationManifestSerializer& ManifestSerializer );
 
 	/** Returns the current culture revision index */
diff --git a/Engine/Source/Runtime/Slate/Public/Widgets/SToolTip.h b/Engine/Source/Runtime/Slate/Public/Widgets/SToolTip.h
--- a/Engine/Source/Runtime/Slate/Public/Widgets/SToolTip.h
+++ b/Engine/Source/Runtime/Slate/Public/Widgets/SToolTip.h
@@ -71,6 +71,11 @@ public:
 	virtual bool IsEmpty( ) const override;
 	virtual bool IsInteractive( ) const override;
 
+	virtual const FText& GetTextTooltip() const
+	{
+		return TextContent.Get();
+	}
+
 public:
 
 	static float GetToolTipWrapWidth();
