Translation Editor changes:

- Translation picker mode that allows you to click on text in the Editor UI and translate it in place
- No longer loads history on startup, can press the "Get History" button in the History tab to retrieve history on demand

diff --git a/Engine/Source/Editor/TranslationEditor/Private/TranslationDataManager.cpp b/Engine/Source/Editor/TranslationEditor/Private/TranslationDataManager.cpp
--- a/Engine/Source/Editor/TranslationEditor/Private/TranslationDataManager.cpp
+++ b/Engine/Source/Editor/TranslationEditor/Private/TranslationDataManager.cpp
@@ -76,9 +76,6 @@ FTranslationDataManager::FTranslationDataManager( const FString& InManifestFileP
 			}
 			GWarn->EndSlowTask();
 
-			// Once we have the context information, we can look up history
-			GetHistoryForTranslationUnits(TranslationUnits, ManifestFilePath);
-
 			LoadFromArchive(TranslationUnits);
 		}
 		else  // ArchivePtr.IsValid() is false
@@ -342,10 +339,16 @@ bool FTranslationDataManager::WriteJSONToTextFile(TSharedRef<FJsonObject>& Outpu
 	return CheckoutAndSaveWasSuccessful;
 }
 
-void FTranslationDataManager::GetHistoryForTranslationUnits( TArray<UTranslationUnit*>& TranslationUnits, const FString& InManifestFilePath )
+void FTranslationDataManager::GetHistoryForTranslationUnits()
 {
 	GWarn->BeginSlowTask(LOCTEXT("LoadingSourceControlHistory", "Loading Translation History from Source Control..."), true);
 
+	TArray<UTranslationUnit*>& TranslationUnits = AllTranslations;
+	const FString& InManifestFilePath = ManifestFilePath;
+
+	// Unload any previous history information, going to retrieve it all again.
+	UnloadHistoryInformation();
+
 	// Force history update
 	ISourceControlProvider& SourceControlProvider = ISourceControlModule::Get().GetProvider();
 	TSharedRef<FUpdateStatus, ESPMode::ThreadSafe> UpdateStatusOperation = ISourceControlOperation::Create<FUpdateStatus>();
@@ -474,6 +477,61 @@ void FTranslationDataManager::GetHistoryForTranslationUnits( TArray<UTranslation
 		TranslationEditorMessageLog.Notify(LOCTEXT("RetrieveTranslationHistoryFailed", "Unable to Retrieve Translation History from Source Control!"));
 	}
 
+
+	// Go though all translation units
+	for (int32 CurrentTranslationUnitIndex = 0; CurrentTranslationUnitIndex < TranslationUnits.Num(); ++CurrentTranslationUnitIndex)
+	{
+		UTranslationUnit* TranslationUnit = TranslationUnits[CurrentTranslationUnitIndex];
+		if (TranslationUnit != nullptr)
+		{
+			if (TranslationUnit->Translation.IsEmpty())
+			{
+				bool bHasTranslationHistory = false;
+				int32 MostRecentNonNullTranslationIndex = -1;
+				int32 ContextForRecentTranslation = -1;
+
+				// Check all contexts for history
+				for (int32 ContextIndex = 0; ContextIndex < TranslationUnit->Contexts.Num(); ++ContextIndex)
+				{
+					for (int32 ChangeIndex = 0; ChangeIndex < TranslationUnit->Contexts[ContextIndex].Changes.Num(); ++ChangeIndex)
+					{
+						if (!(TranslationUnit->Contexts[ContextIndex].Changes[ChangeIndex].Translation.IsEmpty()))
+						{
+							bHasTranslationHistory = true;
+							MostRecentNonNullTranslationIndex = ChangeIndex;
+							ContextForRecentTranslation = ContextIndex;
+							break;
+						}
+					}
+
+					if (bHasTranslationHistory)
+					{
+						break;
+					}
+				}
+
+				// If we have history, but current translation is empty, this goes in the Needs Review tab
+				if (bHasTranslationHistory)
+				{
+					// Offer the most recent translation (for the first context in the list) as a suggestion or starting point (not saved unless user checks "Has Been Reviewed")
+					TranslationUnit->Translation = TranslationUnit->Contexts[ContextForRecentTranslation].Changes[MostRecentNonNullTranslationIndex].Translation;
+					TranslationUnit->HasBeenReviewed = false;
+
+					// Move from Untranslated to review
+					if (Untranslated.Contains(TranslationUnit))
+					{
+						Untranslated.Remove(TranslationUnit);
+					}
+					if (!Review.Contains(TranslationUnit))
+					{
+						Review.Add(TranslationUnit);
+					}
+				}
+			}
+		}
+	}
+
+
 	GWarn->EndSlowTask();
 }
 
@@ -677,4 +735,113 @@ void FTranslationDataManager::RemoveTranslationUnitArrayfromRoot(TArray<UTransla
 	}
 }
 
+void FTranslationDataManager::UnloadHistoryInformation()
+{
+	TArray<UTranslationUnit*>& TranslationUnits = AllTranslations;
+
+	for (int32 CurrentTranslationUnitIndex = 0; CurrentTranslationUnitIndex < TranslationUnits.Num(); ++CurrentTranslationUnitIndex)
+	{
+		UTranslationUnit* TranslationUnit = TranslationUnits[CurrentTranslationUnitIndex];
+		if (TranslationUnit != nullptr)
+		{
+			// If HasBeenReviewed is false, this is a suggestion translation from a previous translation for the same Namespace/Key pair
+			if (!TranslationUnit->HasBeenReviewed)
+			{
+				if (!Untranslated.Contains(TranslationUnit))
+				{
+					Untranslated.Add(TranslationUnit);
+				}
+				if (Review.Contains(TranslationUnit))
+				{
+					Review.Remove(TranslationUnit);
+				}
+
+				// Erase previously suggested translation from history (it has not been reviewed)
+				TranslationUnit->Translation.Empty();
+
+				// Remove all history entries
+				for (FTranslationContextInfo Context : TranslationUnit->Contexts)
+				{
+					Context.Changes.Empty();
+				}
+			}
+		}
+	}
+}
+
+bool FTranslationDataManager::SaveSelectedTranslations(TArray<UTranslationUnit*> TranslationUnitsToSave)
+{
+	bool bSucceeded = true;
+	
+	TMap<FString, TSharedPtr<TArray<UTranslationUnit*>>> TextsToSavePerProject;
+
+	// Regroup the translations to save by project
+	for (UTranslationUnit* TextToSave : TranslationUnitsToSave)
+	{
+		FString LocresFilePath = TextToSave->LocresPath;
+		if (!LocresFilePath.IsEmpty())
+		{
+			if (!TextsToSavePerProject.Contains(LocresFilePath))
+			{
+				TextsToSavePerProject.Add(LocresFilePath, MakeShareable(new TArray<UTranslationUnit*>()));
+			}
+
+			TSharedPtr<TArray<UTranslationUnit*>> ProjectArray = TextsToSavePerProject.FindRef(LocresFilePath);
+			ProjectArray->Add(TextToSave);
+		}
+	}
+
+	for (auto TextIt = TextsToSavePerProject.CreateIterator(); TextIt; ++TextIt)
+	{
+		auto Item = *TextIt;
+		FString CurrentLocResPath = Item.Key;
+		FString ManifestAndArchiveName = FPaths::GetBaseFilename(CurrentLocResPath);
+		FString ArchiveFilePath = FPaths::GetPath(CurrentLocResPath);
+		FString CultureName = FPaths::GetBaseFilename(ArchiveFilePath);
+		FString ManifestPath = FPaths::GetPath(ArchiveFilePath);
+		FString ArchiveFullPath = ArchiveFilePath / ManifestAndArchiveName + ".archive";
+		FString ManifestFullPath = ManifestPath / ManifestAndArchiveName + ".manifest";
+
+		if (FPaths::FileExists(ManifestFullPath) && FPaths::FileExists(ArchiveFullPath))
+		{
+			TSharedRef<FTranslationDataManager> DataManager = MakeShareable(new FTranslationDataManager(ManifestFullPath, ArchiveFullPath));
+
+			TArray<UTranslationUnit*>& TranslationsArray = DataManager->GetAllTranslationsArray();
+			TSharedPtr<TArray<UTranslationUnit*>> EditedItems = Item.Value;
+
+			// For each edited item belonging to this manifest/archive pair
+			for (auto EditedItemIt = EditedItems->CreateIterator(); EditedItemIt; ++EditedItemIt)
+			{
+				UTranslationUnit* EditedItem = *EditedItemIt;
+
+				// Search all translations for the one that matches this FText
+				for (UTranslationUnit* Translation : TranslationsArray)
+				{
+					// If namespace matches...
+					if (Translation->Namespace == EditedItem->Namespace)
+					{
+						// And source matches
+						if (Translation->Source == EditedItem->Source)
+						{
+							// Update the translation in TranslationDataManager, and finish searching these translations
+							Translation->Translation = EditedItem->Translation;
+							break;
+						}
+					}
+				}
+			}
+
+			// Save the data to file, and preview in editor
+			bSucceeded = bSucceeded && DataManager->WriteTranslationData();
+			DataManager->PreviewAllTranslationsInEditor();
+		}
+		else
+		{
+			bSucceeded = false;
+		}
+	}
+
+	return bSucceeded;
+}
+
 #undef LOCTEXT_NAMESPACE
diff --git a/Engine/Source/Editor/TranslationEditor/Private/TranslationDataManager.h b/Engine/Source/Editor/TranslationEditor/Private/TranslationDataManager.h
--- a/Engine/Source/Editor/TranslationEditor/Private/TranslationDataManager.h
+++ b/Engine/Source/Editor/TranslationEditor/Private/TranslationDataManager.h
@@ -67,6 +67,15 @@ public:
 	/** Load (or reload) Translations from Archive file */ 
 	void LoadFromArchive(TArray<UTranslationUnit*>& InTranslationUnits, bool bTrackChanges = false, bool bReloadFromFile = false);
 
+	/** Get the history data for all translation units */
+	void GetHistoryForTranslationUnits();
+
+	/** Unload History information (in the case we reload it) */
+	void UnloadHistoryInformation();
+
+	/** Save the specified translations */
+	static bool SaveSelectedTranslations(TArray<UTranslationUnit*> TranslationUnitsToSave);
+
 private:
 	/** Read text file into a JSON file */
 	static TSharedPtr<FJsonObject> ReadJSONTextFile( const FString& InFilePath ) ;
@@ -79,9 +88,6 @@ private:
 	/** Write JSON file to text file */
 	bool WriteJSONToTextFile( TSharedRef<FJsonObject>& Output, const FString& Filename );
 
-	/** Get the history data for a given translation unit */
-	void GetHistoryForTranslationUnits( TArray<UTranslationUnit*>& TranslationUnits, const FString& ManifestFilePath );
-
 	/** Removes each UTranslationUnit in the passed array from the root set, allowing it to be garbage collected */
 	void RemoveTranslationUnitArrayfromRoot( TArray<UTranslationUnit*>& TranslationUnits );
 
diff --git a/Engine/Source/Editor/TranslationEditor/Private/TranslationEditor.cpp b/Engine/Source/Editor/TranslationEditor/Private/TranslationEditor.cpp
--- a/Engine/Source/Editor/TranslationEditor/Private/TranslationEditor.cpp
+++ b/Engine/Source/Editor/TranslationEditor/Private/TranslationEditor.cpp
@@ -641,9 +641,14 @@ TSharedRef<SDockTab> FTranslationEditor::SpawnTab_Context( const FSpawnTabArgs&
 				+SVerticalBox::Slot()
 				.HAlign(HAlign_Left)
 				.VAlign(VAlign_Center)
-				.FillHeight(0.1)
+				.AutoHeight()
 				[
-					NamespaceTextBlock
+					SNew(SBorder)
+					.BorderImage(FEditorStyle::GetBrush("ToolPanel.GroupBorder"))
+					.Padding(5.0f)
+					[
+						NamespaceTextBlock
+					]
 				]
 				+SVerticalBox::Slot()
 				.HAlign(HAlign_Fill)
@@ -657,28 +662,28 @@ TSharedRef<SDockTab> FTranslationEditor::SpawnTab_Context( const FSpawnTabArgs&
 	return NewDockTab;
 }
 
-TSharedRef<SDockTab> FTranslationEditor::SpawnTab_History( const FSpawnTabArgs& Args )
+TSharedRef<SDockTab> FTranslationEditor::SpawnTab_History(const FSpawnTabArgs& Args)
 {
-	check( Args.GetTabId().TabType == HistoryTabId );
+	check(Args.GetTabId().TabType == HistoryTabId);
 
-	FPropertyEditorModule& PropertyEditorModule = FModuleManager::GetModuleChecked<FPropertyEditorModule>( "PropertyEditor" );
+	FPropertyEditorModule& PropertyEditorModule = FModuleManager::GetModuleChecked<FPropertyEditorModule>("PropertyEditor");
 
-	UProperty* SourceProperty = FindField<UProperty>( FTranslationChange::StaticStruct(), "Source");
-	UProperty* TranslationProperty = FindField<UProperty>( FTranslationChange::StaticStruct(), "Translation");
+	UProperty* SourceProperty = FindField<UProperty>(FTranslationChange::StaticStruct(), "Source");
+	UProperty* TranslationProperty = FindField<UProperty>(FTranslationChange::StaticStruct(), "Translation");
 
 	// create empty property table
 	HistoryPropertyTable = PropertyEditorModule.CreatePropertyTable();
-	HistoryPropertyTable->SetIsUserAllowedToChangeRoot( false );
-	HistoryPropertyTable->SetOrientation( EPropertyTableOrientation::AlignPropertiesInColumns );
-	HistoryPropertyTable->SetShowRowHeader( true );
-	HistoryPropertyTable->SetShowObjectName( false );
+	HistoryPropertyTable->SetIsUserAllowedToChangeRoot(false);
+	HistoryPropertyTable->SetOrientation(EPropertyTableOrientation::AlignPropertiesInColumns);
+	HistoryPropertyTable->SetShowRowHeader(true);
+	HistoryPropertyTable->SetShowObjectName(false);
 
 	// we want to customize some columns
 	TArray< TSharedRef<class IPropertyTableCustomColumn>> CustomColumns;
 	SourceColumn->AddSupportedProperty(SourceProperty);
 	TranslationColumn->AddSupportedProperty(TranslationProperty);
-	CustomColumns.Add( SourceColumn );
-	CustomColumns.Add( TranslationColumn );
+	CustomColumns.Add(SourceColumn);
+	CustomColumns.Add(TranslationColumn);
 
 	if (DataManager->GetAllTranslationsArray().Num() > 0)
 	{
@@ -689,14 +694,14 @@ TSharedRef<SDockTab> FTranslationEditor::SpawnTab_History( const FSpawnTabArgs&
 
 	// Build the Path to the data we want to show
 	TSharedRef<FPropertyPath> Path = FPropertyPath::CreateEmpty();
-	UArrayProperty* ContextsProp = FindField<UArrayProperty>( UTranslationUnit::StaticClass(), "Contexts" );
+	UArrayProperty* ContextsProp = FindField<UArrayProperty>(UTranslationUnit::StaticClass(), "Contexts");
 	Path = Path->ExtendPath(FPropertyPath::Create(ContextsProp));
 	FPropertyInfo ContextsPropInfo;
 	ContextsPropInfo.Property = ContextsProp->Inner;
 	ContextsPropInfo.ArrayIndex = 0;
 	Path = Path->ExtendPath(ContextsPropInfo);
 
-	UProperty* ChangesProp = FindField<UProperty>( FTranslationContextInfo::StaticStruct(), "Changes" );
+	UProperty* ChangesProp = FindField<UProperty>(FTranslationContextInfo::StaticStruct(), "Changes");
 	FPropertyInfo ChangesPropInfo;
 	ChangesPropInfo.Property = ChangesProp;
 	ChangesPropInfo.ArrayIndex = INDEX_NONE;
@@ -704,8 +709,8 @@ TSharedRef<SDockTab> FTranslationEditor::SpawnTab_History( const FSpawnTabArgs&
 	HistoryPropertyTable->SetRootPath(Path);
 
 	// Add the columns we want to display
-	HistoryPropertyTable->AddColumn((TWeakObjectPtr<UProperty>)FindField<UProperty>( FTranslationChange::StaticStruct(), "Version"));
-	HistoryPropertyTable->AddColumn((TWeakObjectPtr<UProperty>)FindField<UProperty>( FTranslationChange::StaticStruct(), "DateAndTime"));
+	HistoryPropertyTable->AddColumn((TWeakObjectPtr<UProperty>)FindField<UProperty>(FTranslationChange::StaticStruct(), "Version"));
+	HistoryPropertyTable->AddColumn((TWeakObjectPtr<UProperty>)FindField<UProperty>(FTranslationChange::StaticStruct(), "DateAndTime"));
 	HistoryPropertyTable->AddColumn((TWeakObjectPtr<UProperty>)SourceProperty);
 	HistoryPropertyTable->AddColumn((TWeakObjectPtr<UProperty>)TranslationProperty);
 
@@ -716,19 +721,44 @@ TSharedRef<SDockTab> FTranslationEditor::SpawnTab_History( const FSpawnTabArgs&
 		Column->SetFrozen(true);
 	}
 
-	HistoryPropertyTableWidgetHandle = PropertyEditorModule.CreatePropertyTableWidgetHandle( HistoryPropertyTable.ToSharedRef(), CustomColumns );
+	HistoryPropertyTableWidgetHandle = PropertyEditorModule.CreatePropertyTableWidgetHandle(HistoryPropertyTable.ToSharedRef(), CustomColumns);
 	TSharedRef<SWidget> PropertyTableWidget = HistoryPropertyTableWidgetHandle->GetWidget();
 
 	TSharedRef<SDockTab> NewDockTab = SNew(SDockTab)
-		.Icon( FEditorStyle::GetBrush("TranslationEditor.Tabs.Properties") )
-		.Label( LOCTEXT("HistoryTabTitle", "History") )
-		.TabColorScale( GetTabColorScale() )
+		.Icon(FEditorStyle::GetBrush("TranslationEditor.Tabs.Properties"))
+		.Label(LOCTEXT("HistoryTabTitle", "History"))
+		.TabColorScale(GetTabColorScale())
 		[
 			SNew(SBorder)
-			.BorderImage( FEditorStyle::GetBrush("ToolPanel.GroupBorder") )
+			.BorderImage(FEditorStyle::GetBrush("ToolPanel.GroupBorder"))
 			.Padding(0.0f)
 			[
-				PropertyTableWidget
+				SNew(SVerticalBox)
+				+ SVerticalBox::Slot()
+				.HAlign(HAlign_Left)
+				.VAlign(VAlign_Center)
+				.AutoHeight()
+				[
+					SNew(SBorder)
+					.BorderImage(FEditorStyle::GetBrush("ToolPanel.GroupBorder"))
+					.Padding(5.0f)
+					[
+						SNew(SButton)
+						.HAlign(HAlign_Center)
+						.VAlign(VAlign_Center)
+						.OnClicked(FOnClicked::CreateSP(this, &FTranslationEditor::OnGetHistoryButtonClicked))
+						[
+							SNew(STextBlock)
+							.Text((LOCTEXT("GetHistoryButton", "Get History...")))
+						]
+					]
+				]
+				+ SVerticalBox::Slot()
+				.HAlign(HAlign_Fill)
+				.VAlign(VAlign_Fill)
+				[
+					PropertyTableWidget
+				]
 			]
 		];
 	
@@ -1446,4 +1476,60 @@ void FTranslationEditor::OpenSearchTab_Execute()
 	TabManager->InvokeTab(SearchTabId);
 }
 
+FReply FTranslationEditor::OnGetHistoryButtonClicked()
+{
+	// Load the actual history data
+	DataManager->GetHistoryForTranslationUnits();
+
+	// Items might have moved from Untranslated to review, so refresh the view of both tables
+	if (UntranslatedPropertyTable.IsValid())
+	{
+		UntranslatedPropertyTable->SetObjects((TArray<UObject*>&)DataManager->GetUntranslatedArray());
+
+		// Need to re-add the columns we want to display
+		UntranslatedPropertyTable->AddColumn((TWeakObjectPtr<UProperty>)FindField<UProperty>(UTranslationUnit::StaticClass(), "Source"));
+		UntranslatedPropertyTable->AddColumn((TWeakObjectPtr<UProperty>)FindField<UProperty>(UTranslationUnit::StaticClass(), "Translation"));
+
+		TArray<TSharedRef<IPropertyTableColumn>> Columns = UntranslatedPropertyTable->GetColumns();
+		for (TSharedRef<IPropertyTableColumn> Column : Columns)
+		{
+			Column->SetFrozen(true);
+		}
+	}
+
+	if (ReviewPropertyTable.IsValid())
+	{
+		ReviewPropertyTable->SetObjects((TArray<UObject*>&)DataManager->GetReviewArray());
+
+		// Need to re-add the columns we want to display
+		ReviewPropertyTable->AddColumn((TWeakObjectPtr<UProperty>)FindField<UProperty>(UTranslationUnit::StaticClass(), "Source"));
+		ReviewPropertyTable->AddColumn((TWeakObjectPtr<UProperty>)FindField<UProperty>(UTranslationUnit::StaticClass(), "Translation"));
+		ReviewPropertyTable->AddColumn((TWeakObjectPtr<UProperty>)FindField<UProperty>(UTranslationUnit::StaticClass(), "HasBeenReviewed"));
+
+		TArray<TSharedRef<IPropertyTableColumn>> Columns = ReviewPropertyTable->GetColumns();
+		for (TSharedRef<IPropertyTableColumn> Column : Columns)
+		{
+			FString ColumnId = Column->GetId().ToString();
+			if (ColumnId == "HasBeenReviewed")
+			{
+				Column->SetWidth(120);
+				Column->SetSizeMode(EPropertyTableColumnSizeMode::Fixed);
+			}
+			// Freeze columns, don't want user to remove them
+			Column->SetFrozen(true);
+		}
+	}
+
+	// Make sure all UI is refreshed
+	RefreshUI();
+
+	// Make sure current selection is reflected
+	UpdateUntranslatedSelection();
+	UpdateNeedsReviewSelection();
+	UpdateCompletedSelection();
+	UpdateSearchSelection();
+
+	return FReply::Handled();
+}
+
 #undef LOCTEXT_NAMESPACE
diff --git a/Engine/Source/Editor/TranslationEditor/Private/TranslationEditor.h b/Engine/Source/Editor/TranslationEditor/Private/TranslationEditor.h
--- a/Engine/Source/Editor/TranslationEditor/Private/TranslationEditor.h
+++ b/Engine/Source/Editor/TranslationEditor/Private/TranslationEditor.h
@@ -176,6 +176,9 @@ private:
 	/** Called when text is committed to the SearchBox */
 	void OnFilterTextCommitted(const FText& InFilterText, ETextCommit::Type CommitInfo);
 
+	/** Called when the get history button is clicked to retrieve history from source control */
+	FReply OnGetHistoryButtonClicked();
+
 	/**	The tab id for the untranslated tab */
 	static const FName UntranslatedTabId;
 	/**	The tab id for the review tab */
diff --git a/Engine/Source/Editor/TranslationEditor/Private/TranslationPickerEditWindow.cpp b/Engine/Source/Editor/TranslationEditor/Private/TranslationPickerEditWindow.cpp
new file mode 100644
--- /dev/null
+++ b/Engine/Source/Editor/TranslationEditor/Private/TranslationPickerEditWindow.cpp
@@ -0,0 +1,268 @@
+
+#include "TranslationEditorPrivatePCH.h"
+#include "TranslationPickerEditWindow.h"
+#include "Editor/Documentation/Public/SDocumentationToolTip.h"
+#include "TranslationUnit.h"
+
+#define LOCTEXT_NAMESPACE "TranslationPicker"
+
+void STranslationPickerEditWindow::Construct(const FArguments& InArgs)
+{
+	ParentWindow = InArgs._ParentWindow;
+	PickedTexts = InArgs._PickedTexts;
+
+	WindowContents = SNew(SBox);
+
+	auto TextsBox = SNew(SVerticalBox);
+
+	// Add a new Translation Picker Edit Widget for each picked text
+	for (FText PickedText : PickedTexts)
+	{
+		TSharedPtr<SEditableTextBox> TextBox;
+		int32 DefaultPadding = 0.0f;
+
+		TextsBox->AddSlot()
+			.FillHeight(1)
+			.Padding(FMargin(5))
+			[
+				SNew(SBorder)
+				[
+					SNew(STranslationPickerEditWidget)
+					.PickedText(PickedText)
+				]
+			];
+	}
+
+	TSharedPtr<SEditableTextBox> TextBox;
+	int32 DefaultPadding = 0.0f;
+
+	// Layout the Translation Picker Edit Widgets and some save/close buttons below them
+	WindowContents->SetContent(
+		SNew(SBorder)
+		.BorderImage(FEditorStyle::GetBrush("ToolPanel.GroupBorder"))
+		[
+			SNew(SVerticalBox)
+			+ SVerticalBox::Slot()
+			.FillHeight(1)
+			[
+				SNew(SBox)
+				.Padding(FMargin(8, 5, 8, 5))
+				[
+					SNew(SVerticalBox)
+					+ SVerticalBox::Slot()
+					.Padding(DefaultPadding)
+					.AutoHeight()
+					[
+						TextsBox
+					]
+					+ SVerticalBox::Slot()
+					.Padding(DefaultPadding)
+					.AutoHeight()
+					.HAlign(HAlign_Right)
+					[
+						SNew(SUniformGridPanel)
+						.SlotPadding(FEditorStyle::GetMargin("StandardDialog.SlotPadding"))
+						.MinDesiredSlotWidth(FEditorStyle::GetFloat("StandardDialog.MinDesiredSlotWidth"))
+						.MinDesiredSlotHeight(FEditorStyle::GetFloat("StandardDialog.MinDesiredSlotHeight"))
+						+ SUniformGridPanel::Slot(0, 0)
+						[
+							SNew(SButton)
+							.HAlign(HAlign_Center)
+							.ContentPadding(FEditorStyle::GetMargin("StandardDialog.ContentPadding"))
+							.OnClicked(this, &STranslationPickerEditWindow::Close)
+							[
+								SNew(SHorizontalBox)
+								+ SHorizontalBox::Slot()
+								.Padding(FMargin(0, 0, 3, 0))
+								.VAlign(VAlign_Center)
+								.AutoWidth()
+								[
+									SNew(STextBlock).Text(LOCTEXT("SaveAllAndClose", "Save all and close"))
+								]
+							]
+						]
+						+ SUniformGridPanel::Slot(1, 0)
+						[
+							SNew(SButton)
+							.HAlign(HAlign_Center)
+							.ContentPadding(FEditorStyle::GetMargin("StandardDialog.ContentPadding"))
+							.OnClicked(this, &STranslationPickerEditWindow::Close)
+							.Text(LOCTEXT("CancelButton", "Cancel"))
+						]
+					]
+				]
+			]
+		]
+
+	);
+
+	ChildSlot
+	[
+		WindowContents.ToSharedRef()
+	];
+}
+
+FReply STranslationPickerEditWindow::OnKeyDown(const FGeometry& MyGeometry, const FKeyEvent& InKeyEvent)
+{
+	if (InKeyEvent.GetKey() == EKeys::Escape)
+	{
+		Close();
+		
+		return FReply::Handled();
+	}
+
+	return FReply::Unhandled();
+}
+
+FReply STranslationPickerEditWindow::Close()
+{
+	if (ParentWindow.IsValid())
+	{
+		FSlateApplication::Get().RequestDestroyWindow(ParentWindow.Pin().ToSharedRef());
+		ParentWindow.Reset();
+	}
+
+	return FReply::Handled();
+}
+
+void STranslationPickerEditWidget::Construct(const FArguments& InArgs)
+{
+	PickedText = InArgs._PickedText;
+	int32 DefaultPadding = 0.0f;
+
+	// Get all the data we need and format it properly
+	const FString* NamespaceString = FTextInspector::GetNamespace(PickedText);
+	const FString* KeyString = FTextInspector::GetKey(PickedText);
+	const FString* SourceString = FTextInspector::GetSourceString(PickedText);
+	const FString& TranslationString = FTextInspector::GetDisplayString(PickedText);
+	FString LocresFullPath;
+
+	FString ManifestAndArchiveNameString;
+	if (NamespaceString && KeyString)
+	{
+		TSharedPtr<FString, ESPMode::ThreadSafe> TextTableName = TSharedPtr<FString, ESPMode::ThreadSafe>(nullptr);
+		TextTableName = FTextLocalizationManager::Get().GetTableName(*NamespaceString, *KeyString);
+		LocresFullPath = *TextTableName;
+		if (TextTableName.IsValid())
+		{
+			ManifestAndArchiveNameString = FPaths::GetBaseFilename(*TextTableName);
+		}
+	}
+
+	FText Namespace = NamespaceString != nullptr ? FText::FromString(*NamespaceString) : FText::GetEmpty();
+	FText Key = KeyString != nullptr ? FText::FromString(*KeyString) : FText::GetEmpty();
+	FText Source = SourceString != nullptr ? FText::FromString(*SourceString) : FText::GetEmpty();
+	FText ManifestAndArchiveName = FText::FromString(ManifestAndArchiveNameString);
+	FText Translation = FText::FromString(TranslationString);
+
+	FText SourceWithLabel = FText::Format(LOCTEXT("SourceLabel", "Source: {0}"), Source);
+	FText NamespaceWithLabel = FText::Format(LOCTEXT("NamespaceLabel", "Namespace: {0}"), Namespace);
+	FText KeyWithLabel = FText::Format(LOCTEXT("KeyLabel", "Key: {0}"), Key);
+	FText ManifestAndArchiveNameWithLabel = FText::Format(LOCTEXT("LocresFileLabel", "Manifest/Archive File: {0}"), FText::FromString(ManifestAndArchiveNameString));
+
+	// Save the necessary data in UTranslationUnit for later.  This is what we pass to TranslationDataManager to save our edits
+	TranslationUnit = NewObject<UTranslationUnit>();
+	TranslationUnit->Namespace = NamespaceString != nullptr ? *NamespaceString : "";
+	TranslationUnit->Source = SourceString != nullptr ? *SourceString : "";
+	TranslationUnit->Translation = TranslationString;
+	TranslationUnit->LocresPath = LocresFullPath;
+
+	// Layout all our data
+	ChildSlot
+		[
+			SNew(SVerticalBox)
+			+SVerticalBox::Slot()
+			.FillHeight(1)
+			.Padding(FMargin(5))
+			[
+				SNew(SVerticalBox)
+				+ SVerticalBox::Slot()
+				.FillHeight(1)
+				.Padding(FMargin(5))
+				[
+					SNew(STextBlock)
+					.Text(SourceWithLabel)
+				]
+				+ SVerticalBox::Slot()
+				.FillHeight(1)
+				.Padding(FMargin(5))
+				[
+					SNew(SHorizontalBox)
+					+ SHorizontalBox::Slot()
+					.AutoWidth()
+					[
+						SNew(STextBlock)
+						.Text(LOCTEXT("TranslationLabel", "Translation: "))
+					]
+					+ SHorizontalBox::Slot()
+					[
+						SAssignNew(TextBox, SEditableTextBox)
+						.Text(Translation)
+						.HintText(LOCTEXT("TranslationEditTextBox_HintText", "Enter/edit translation here."))
+					]
+				]
+			]
+			+SVerticalBox::Slot()
+			.FillHeight(1)
+			.Padding(FMargin(5))
+			[
+				SNew(SHorizontalBox)
+				+ SHorizontalBox::Slot()
+				.AutoWidth()
+				.Padding(FMargin(5))
+				[
+					SNew(SVerticalBox)
+					+ SVerticalBox::Slot()
+					[
+						SNew(STextBlock)
+						.Text(NamespaceWithLabel)
+					]
+					+ SVerticalBox::Slot()
+					[
+						SNew(STextBlock)
+						.Text(KeyWithLabel)
+					]
+					+ SVerticalBox::Slot()
+					[
+						SNew(STextBlock)
+						.Text(ManifestAndArchiveNameWithLabel)
+					]
+				]
+				+ SHorizontalBox::Slot()
+				.AutoWidth()
+				.Padding(FMargin(5))
+				[
+					SNew(SButton)
+					.HAlign(HAlign_Center)
+					.ContentPadding(FEditorStyle::GetMargin("StandardDialog.ContentPadding"))
+					.OnClicked(this, &STranslationPickerEditWidget::SaveAndPreview)
+					[
+						SNew(SHorizontalBox)
+						+ SHorizontalBox::Slot()
+						.Padding(FMargin(0, 0, 3, 0))
+						.VAlign(VAlign_Center)
+						.AutoWidth()
+						[
+							SNew(STextBlock)
+							.Text(LOCTEXT("SaveAndPreviewButtonText", "Save and preview"))
+						]
+					]
+				]
+			]
+		];
+}
+
+FReply STranslationPickerEditWidget::SaveAndPreview()
+{
+	// Update translation string from entered text
+	TranslationUnit->Translation = TextBox->GetText().ToString();
+
+	// Save the data via translation data manager
+	TArray<UTranslationUnit*> TempArray;
+	TempArray.Add(TranslationUnit);
+	FTranslationDataManager::SaveSelectedTranslations(TempArray);
+
+	return FReply::Handled();
+}
+
+#undef LOCTEXT_NAMESPACE
\ No newline at end of file
diff --git a/Engine/Source/Editor/TranslationEditor/Private/TranslationPickerEditWindow.h b/Engine/Source/Editor/TranslationEditor/Private/TranslationPickerEditWindow.h
new file mode 100644
--- /dev/null
+++ b/Engine/Source/Editor/TranslationEditor/Private/TranslationPickerEditWindow.h
@@ -0,0 +1,71 @@
+#pragma once
+#include "SCompoundWidget.h"
+
+#define LOCTEXT_NAMESPACE "TranslationPicker"
+
+/** Translation picker edit Widget to handle the display and editing of a single selected FText */
+class STranslationPickerEditWidget : public SCompoundWidget
+{
+	SLATE_BEGIN_ARGS(STranslationPickerEditWidget) {}
+
+	SLATE_ARGUMENT(FText, PickedText)
+
+	SLATE_END_ARGS()
+
+	void Construct(const FArguments& InArgs);
+
+private:
+
+	virtual bool SupportsKeyboardFocus() const override
+	{
+		return true;
+	}
+
+	FReply SaveAndPreview();
+
+	/** The FText that we are using this widget to translate */
+	FText PickedText;
+
+	/** The translation we're editing represented in a UTranslationUnit object */
+	UTranslationUnit* TranslationUnit;
+
+	/** The text box for entering/modifying a translation */
+	TSharedPtr<SEditableTextBox> TextBox;
+};
+
+/** Translation picker edit window to allow you to translate selected FTexts in place */
+class STranslationPickerEditWindow : public SCompoundWidget
+{
+	SLATE_BEGIN_ARGS(STranslationPickerEditWindow) {}
+
+	SLATE_ARGUMENT(TWeakPtr<SWindow>, ParentWindow)
+
+	SLATE_ARGUMENT(TArray<FText>, PickedTexts)
+
+	SLATE_END_ARGS()
+
+	void Construct(const FArguments& InArgs);
+
+private:
+
+	virtual FReply OnKeyDown(const FGeometry& MyGeometry, const FKeyEvent& InKeyEvent) override;
+
+	FReply Close();
+
+	/** We need to support keyboard focus to process the 'Esc' key */
+	virtual bool SupportsKeyboardFocus() const override
+	{
+		return true;
+	}
+
+	/** Handle to the window that contains this widget */
+	TWeakPtr<SWindow> ParentWindow;
+
+	/** Contents of the window */
+	TSharedPtr<SBox> WindowContents;
+
+	/** The FTexts that we have found under the cursor */
+	TArray<FText> PickedTexts;
+};
+
+#undef LOCTEXT_NAMESPACE
\ No newline at end of file
diff --git a/Engine/Source/Editor/TranslationEditor/Private/TranslationPickerFloatingWindow.cpp b/Engine/Source/Editor/TranslationEditor/Private/TranslationPickerFloatingWindow.cpp
--- a/Engine/Source/Editor/TranslationEditor/Private/TranslationPickerFloatingWindow.cpp
+++ b/Engine/Source/Editor/TranslationEditor/Private/TranslationPickerFloatingWindow.cpp
@@ -2,6 +2,7 @@
 #include "TranslationEditorPrivatePCH.h"
 #include "TranslationPickerFloatingWindow.h"
 #include "Editor/Documentation/Public/SDocumentationToolTip.h"
+#include "TranslationPickerEditWindow.h"
 
 #define LOCTEXT_NAMESPACE "TranslationPicker"
 
@@ -9,15 +10,22 @@ void STranslationPickerFloatingWindow::Construct(const FArguments& InArgs)
 {
 	ParentWindow = InArgs._ParentWindow;
 
+	WindowContents = SNew(SBox);
+
+	WindowContents->SetContent(
+		SNew(SToolTip)
+		.Text(this, &STranslationPickerFloatingWindow::GetPickerStatusText)
+		);
+
 	ChildSlot
-		[
-			SNew(SToolTip)
-			.Text(this, &STranslationPickerFloatingWindow::GetPickerStatusText)
-		];
+	[
+		WindowContents.ToSharedRef()
+	];
 }
 
 void STranslationPickerFloatingWindow::Tick(const FGeometry& AllottedGeometry, const double InCurrentTime, const float InDeltaTime)
 {
+	PickedTexts.Empty();
 	TranslationInfoPreviewText = FText::GetEmpty();
 
 	FWidgetPath Path = FSlateApplication::Get().LocateWindowUnderMouse(FSlateApplication::Get().GetCursorPos(), FSlateApplication::Get().GetInteractiveTopLevelWindows(), true);
@@ -31,7 +39,7 @@ void STranslationPickerFloatingWindow::Tick(const FGeometry& AllottedGeometry, c
 
 		if (!TextBlockDescription.IsEmpty())
 		{
-			TranslationInfoPreviewText = FText::Format(LOCTEXT("TextBlockLocalizationDescription", "{0}\n\n\n\n\nTextBlockInfo: {1}"),
+			TranslationInfoPreviewText = FText::Format(LOCTEXT("TextBlockLocalizationDescription", "{0}\n\n\n\n\nTextBlock info:\n{1}"),
 				TranslationInfoPreviewText, TextBlockDescription);
 		}
 
@@ -44,7 +52,7 @@ void STranslationPickerFloatingWindow::Tick(const FGeometry& AllottedGeometry, c
 			TooltipDescription = GetTextDescription(Tooltip->AsWidget());
 			if (!TooltipDescription.IsEmpty())
 			{
-				TranslationInfoPreviewText = FText::Format(LOCTEXT("TooltipLocalizationDescription", "{0}\n\n\n\n\nTooltip Info:{1}"),
+				TranslationInfoPreviewText = FText::Format(LOCTEXT("TooltipLocalizationDescription", "{0}\n\n\n\n\nTooltip Info:\n{1}"),
 					TranslationInfoPreviewText, TooltipDescription);
 			}
 		}
@@ -62,14 +70,15 @@ void STranslationPickerFloatingWindow::Tick(const FGeometry& AllottedGeometry, c
 
 FText STranslationPickerFloatingWindow::GetTextDescription(TSharedRef<SWidget> Widget)
 {
-	FText OutText = FText::GetEmpty();
+	FText OriginalText = FText::GetEmpty();
+	FText FormattedText = FText::GetEmpty();
 
 	// Parse STextBlocks
 	STextBlock* TextBlock = (STextBlock*)&Widget.Get();
 	if ((Widget->GetTypeAsString() == "STextBlock" && TextBlock))
 	{
-		FText Text = TextBlock->GetText();
-		OutText = FormatFTextInfo(Text);
+		OriginalText = TextBlock->GetText();
+		FormattedText = FormatFTextInfo(OriginalText);
 	}
 
 	// Parse SToolTips
@@ -78,10 +87,11 @@ FText STranslationPickerFloatingWindow::GetTextDescription(TSharedRef<SWidget> W
 		SToolTip* ToolTip = ((SToolTip*)&Widget.Get());
 		if (ToolTip != nullptr)
 		{
-			OutText = GetTextDescription(ToolTip->GetContentWidget());
-			if (OutText.IsEmpty())
+			FormattedText = GetTextDescription(ToolTip->GetContentWidget());
+			if (FormattedText.IsEmpty())
 			{
-				OutText = FormatFTextInfo(ToolTip->GetTextTooltip());
+				OriginalText = ToolTip->GetTextTooltip();
+				FormattedText = FormatFTextInfo(OriginalText);
 			}
 		}
 	}
@@ -92,10 +102,17 @@ FText STranslationPickerFloatingWindow::GetTextDescription(TSharedRef<SWidget> W
 		SDocumentationToolTip* DocumentationToolTip = (SDocumentationToolTip*)&Widget.Get();
 		if (DocumentationToolTip != nullptr)
 		{
-			OutText = FormatFTextInfo(DocumentationToolTip->GetTextTooltip());
+			OriginalText = DocumentationToolTip->GetTextTooltip();
+			FormattedText = FormatFTextInfo(OriginalText);
 		}
 	}
-	return OutText;
+	
+	if (!OriginalText.IsEmpty())
+	{
+		PickedTexts.Add(OriginalText);
+	}
+
+	return FormattedText;
 }
 
 FText STranslationPickerFloatingWindow::FormatFTextInfo(FText TextToFormat)
@@ -121,7 +138,7 @@ FText STranslationPickerFloatingWindow::FormatFTextInfo(FText TextToFormat)
 		FText TableName = TextTableName.IsValid() ? FText::FromString(*(TextTableName.Get())) : FText::GetEmpty();
 		FText Translation = FText::FromString(TextTranslation);
 
-		OutText = FText::Format(LOCTEXT("LocalizationStringDescription", "\n\n\n\n\n\nText Namespace: {0}\nText Key: {1}\nText Source: {2}\nTable Name: {3}"), Namespace, Key, Source, TableName);
+		OutText = FText::Format(LOCTEXT("LocalizationStringDescription", "Text Namespace: {0}\nText Key: {1}\nText Source: {2}\nTable Name: {3}"), Namespace, Key, Source, TableName);
 		OutText = FText::Format(LOCTEXT("LocalizationStringDescription2", "{0}\nTranslation: {1}"), OutText, Translation);
 	}
 
@@ -132,21 +149,51 @@ FReply STranslationPickerFloatingWindow::OnKeyDown(const FGeometry& MyGeometry,
 {
 	if (InKeyEvent.GetKey() == EKeys::Escape)
 	{
-		//TODO: pop up window for in-place translation
-
-		TranslationInfoPreviewText = FText::GetEmpty();
-
-		if (ParentWindow.IsValid())
+		// Open a different window to allow editing of the translation
+		TSharedRef<SWindow> NewWindow = SNew(SWindow)
+		.Title(LOCTEXT("TranslationPickerEditWindowTitle", "Edit Translation(s)"))
+		.SupportsMaximize(false)
+		.SupportsMinimize(false)
+		.CreateTitleBar(true)
+		.SizingRule(ESizingRule::Autosized);
+		NewWindow->MoveWindowTo(FSlateApplication::Get().GetCursorPos());
+
+		NewWindow->SetContent(
+			SNew(STranslationPickerEditWindow)
+			.ParentWindow(NewWindow)
+			.PickedTexts(PickedTexts)
+			);
+
+		TSharedPtr<SWindow> RootWindow = FGlobalTabmanager::Get()->GetRootWindow();
+		if (RootWindow.IsValid())
+		{
+			FSlateApplication::Get().AddWindowAsNativeChild(NewWindow, RootWindow.ToSharedRef());
+		}
+		else
 		{
-			FSlateApplication::Get().RequestDestroyWindow(ParentWindow.Pin().ToSharedRef());
-			ParentWindow.Reset();
+			FSlateApplication::Get().AddWindow(NewWindow);
 		}
 
+		Close();
+
+		TranslationInfoPreviewText = FText::GetEmpty();
+		
 		return FReply::Handled();
 	}
 
 	return FReply::Unhandled();
 }
 
+FReply STranslationPickerFloatingWindow::Close()
+{
+	if (ParentWindow.IsValid())
+	{
+		FSlateApplication::Get().RequestDestroyWindow(ParentWindow.Pin().ToSharedRef());
+		ParentWindow.Reset();
+	}
+
+	return FReply::Handled();
+}
+
 
 #undef LOCTEXT_NAMESPACE
\ No newline at end of file
diff --git a/Engine/Source/Editor/TranslationEditor/Private/TranslationPickerFloatingWindow.h b/Engine/Source/Editor/TranslationEditor/Private/TranslationPickerFloatingWindow.h
--- a/Engine/Source/Editor/TranslationEditor/Private/TranslationPickerFloatingWindow.h
+++ b/Engine/Source/Editor/TranslationEditor/Private/TranslationPickerFloatingWindow.h
@@ -3,7 +3,7 @@
 
 #define LOCTEXT_NAMESPACE "TranslationPicker"
 
-/** Translation picker window to show details of FText(s) under cursor, and allow in-place translation */
+/** Translation picker floating window to show details of FText(s) under cursor, and allow in-place translation via TranslationPickerEditWindow */
 class STranslationPickerFloatingWindow : public SCompoundWidget
 {
 	SLATE_BEGIN_ARGS(STranslationPickerFloatingWindow) {}
@@ -14,11 +14,6 @@ class STranslationPickerFloatingWindow : public SCompoundWidget
 
 	void Construct(const FArguments& InArgs);
 
-	FText GetTranslationInfoPreviewText()
-	{
-		return TranslationInfoPreviewText;
-	}
-
 private:
 	virtual void Tick(const FGeometry& AllottedGeometry, const double InCurrentTime, const float InDeltaTime) override;
 
@@ -28,13 +23,18 @@ private:
 	/** Format an FText for display */
 	FText FormatFTextInfo(FText TextToFormat);
 
+	/** This is what this window displays */
 	FText GetPickerStatusText() const
 	{
 		return FText::Format(LOCTEXT("TootipHint", "{0}\n\n(Esc to pick)"), TranslationInfoPreviewText);
 	}
 
+	/** Handle key presses */
 	virtual FReply OnKeyDown(const FGeometry& MyGeometry, const FKeyEvent& InKeyEvent) override;
 
+	/** Close this window */
+	FReply Close();
+
 	/** We need to support keyboard focus to process the 'Esc' key */
 	virtual bool SupportsKeyboardFocus() const override
 	{
@@ -44,7 +44,13 @@ private:
 	/** Handle to the window that contains this widget */
 	TWeakPtr<SWindow> ParentWindow;
 
-	/** The widget name we are picking */
+	/** Contents of the window */
+	TSharedPtr<SBox> WindowContents;
+
+	/** The FTexts that we have found under the cursor */
+	TArray<FText> PickedTexts;
+
+	/** Preview text for the translation */
 	FText TranslationInfoPreviewText;
 };
 
diff --git a/Engine/Source/Editor/TranslationEditor/Private/TranslationUnit.h b/Engine/Source/Editor/TranslationEditor/Private/TranslationUnit.h
--- a/Engine/Source/Editor/TranslationEditor/Private/TranslationUnit.h
+++ b/Engine/Source/Editor/TranslationEditor/Private/TranslationUnit.h
@@ -78,6 +78,10 @@ public:
 	UPROPERTY(Category=Translation, EditAnywhere)
 	FString TranslationBeforeImport;
 
+	/** Optional: Which Locres File this translation is in */
+	UPROPERTY(Category = Translation, EditAnywhere)
+	FString LocresPath;
+
 	/**
 	* Returns an event delegate that is executed when a property has changed.
 	*
