Translation Editor: Keep track of open translation editors and prevent the accidental simultaneous edit of the same archive in multiple editors.

diff --git a/Engine/Source/Editor/TranslationEditor/Private/ITranslationEditor.cpp b/Engine/Source/Editor/TranslationEditor/Private/ITranslationEditor.cpp
new file mode 100644
--- /dev/null
+++ b/Engine/Source/Editor/TranslationEditor/Private/ITranslationEditor.cpp
@@ -0,0 +1,46 @@
+// Copyright 1998-2015 Epic Games, Inc. All Rights Reserved.
+
+
+#include "TranslationEditorPrivatePCH.h"
+#include "ITranslationEditor.h"
+
+/** To keep track of what translation editors are open editing which manifest files */
+TMap<FString, ITranslationEditor*> ITranslationEditor::OpenTranslationEditors = TMap<FString, ITranslationEditor*>();
+
+void ITranslationEditor::OpenTranslationEditor(const FString& InManifestFile, const FString& InArchiveFile)
+{
+	// Only open a new translation editor if there isn't another one open for this archive file already
+	if (!OpenTranslationEditors.Contains(InArchiveFile))
+	{
+		FTranslationEditorModule& TranslationEditorModule = FModuleManager::LoadModuleChecked<FTranslationEditorModule>("TranslationEditor");
+		ITranslationEditor* NewTranslationEditor = (ITranslationEditor*)&(TranslationEditorModule.CreateTranslationEditor(InManifestFile, InArchiveFile).Get());
+		NewTranslationEditor->RegisterTranslationEditor();
+	}
+	else
+	{
+		// If already editing this archive file, flash the tab that contains the editor that has that file open
+		ITranslationEditor* OpenEditor = *OpenTranslationEditors.Find(InArchiveFile);
+		OpenEditor->FocusWindow();
+	}
+}
+
+void ITranslationEditor::RegisterTranslationEditor()
+{
+	ITranslationEditor::OpenTranslationEditors.Add(ArchiveFilePath, this);
+}
+
+void ITranslationEditor::UnregisterTranslationEditor()
+{
+	OpenTranslationEditors.Remove(ArchiveFilePath);
+}
+
+bool ITranslationEditor::OnRequestClose()
+{
+	ITranslationEditor::UnregisterTranslationEditor();
+
+	return true;
+}
+
+
+
+
diff --git a/Engine/Source/Editor/TranslationEditor/Private/TranslationEditor.h b/Engine/Source/Editor/TranslationEditor/Private/TranslationEditor.h
--- a/Engine/Source/Editor/TranslationEditor/Private/TranslationEditor.h
+++ b/Engine/Source/Editor/TranslationEditor/Private/TranslationEditor.h
@@ -59,7 +59,8 @@ protected:
 private:
 
 	FTranslationEditor(TSharedRef< FTranslationDataManager > InDataManager, const FString& InManifestFile, const FString& InArchiveFile )
-	:  DataManager(InDataManager)
+		: ITranslationEditor(InManifestFile, InArchiveFile)
+	, DataManager(InDataManager)
 	, SourceFont(FEditorStyle::GetFontStyle( PropertyTableConstants::NormalFontStyle ))
 	, TranslationTargetFont(FEditorStyle::GetFontStyle( PropertyTableConstants::NormalFontStyle ))
 	, SourceColumn(MakeShareable(new FCustomFontColumn(SourceFont)))
@@ -69,8 +70,6 @@ private:
 				.Font(TranslationTargetFont))
 	, NamespaceTextBlock(SNew(STextBlock)
 				.Text(FText::FromString("")))
-	, ManifestFilePath(InManifestFile)
-	, ArchiveFilePath(InArchiveFile)
 	{}
 
 	/** Does some things we can't do in the constructor because we can't get a SharedRef to "this" there */ 
@@ -260,11 +259,6 @@ private:
 	/** Text block displaying the namespace of the currently selected translation unit */
 	TSharedRef<STextBlock> NamespaceTextBlock;
 
-	/** Name of the project we are translating for */
-	FString ManifestFilePath;
-	/** Name of the language we are translating to */
-	FString ArchiveFilePath;
-
 	/** Used to remember the location of the file the user last exported to */
 	FString LastExportFilePath;
 	/** Used to remember the location of the file the user last imported */
diff --git a/Engine/Source/Editor/TranslationEditor/Public/ITranslationEditor.h b/Engine/Source/Editor/TranslationEditor/Public/ITranslationEditor.h
--- a/Engine/Source/Editor/TranslationEditor/Public/ITranslationEditor.h
+++ b/Engine/Source/Editor/TranslationEditor/Public/ITranslationEditor.h
@@ -13,12 +13,32 @@ class ITranslationEditor : public FAssetEditorToolkit
 
 public:
 
-	static void OpenTranslationEditor(const FString& ManifestFile, const FString& ArchiveFile)
-	{
-		FTranslationEditorModule& TranslationEditorModule = FModuleManager::LoadModuleChecked<FTranslationEditorModule>( "TranslationEditor" );
-		TSharedRef< FTranslationEditor > NewTranslationEditor = TranslationEditorModule.CreateTranslationEditor(ManifestFile, ArchiveFile);
-	}
+	TRANSLATIONEDITOR_API static void OpenTranslationEditor(const FString& InManifestFile, const FString& InArchiveFile);
 
+	ITranslationEditor(const FString& InManifestFile, const FString& InArchiveFile)
+		: ManifestFilePath(InManifestFile)
+		, ArchiveFilePath(InArchiveFile)
+	{}
+
+	virtual bool OnRequestClose() override;
+
+protected:
+
+	/** Name of the project we are translating for */
+	FString ManifestFilePath;
+	/** Name of the language we are translating to */
+	FString ArchiveFilePath;
+
+private:
+
+	/** To keep track of what translation editors are open editing which manifest files */
+	static TMap<FString, ITranslationEditor*> OpenTranslationEditors;
+	
+	/** Called on open to add this to our list of open translation editors */
+	void RegisterTranslationEditor();
+	
+	/** Called on close to remove this from our list of open translation editors */
+	void UnregisterTranslationEditor();
 };
 
 
