Translation Picker improvements:

- Can now start the translation picker from the main Window menu in the editor
- Disable "save" button for text that we can't find the manifest/archive data for
- Hover window now presents the information about an FText in a more readable format
- Fixed issue where the same FText info was not displayed twice in certain cases

diff --git a/Engine/Source/Editor/MainFrame/Private/Menus/MainMenu.cpp b/Engine/Source/Editor/MainFrame/Private/Menus/MainMenu.cpp
--- a/Engine/Source/Editor/MainFrame/Private/Menus/MainMenu.cpp
+++ b/Engine/Source/Editor/MainFrame/Private/Menus/MainMenu.cpp
@@ -14,6 +14,7 @@
 #include "GenericCommands.h"
 #include "ToolboxModule.h"
 
+
 #define LOCTEXT_NAMESPACE "MainFileMenu"
 
 
@@ -238,6 +239,12 @@ void FMainMenu::FillWindowMenu( FMenuBuilder& MenuBuilder, const TSharedRef< FEx
 						LOCTEXT("EditorPreferencesSubMenuToolTip", "Open the Translation Editor for a Given Project and Language"),
 						FNewMenuDelegate::CreateStatic(&FMainFrameTranslationEditorMenu::MakeMainFrameTranslationEditorSubMenu)
 						);
+					MenuBuilder.AddMenuEntry(
+						LOCTEXT("TranslationPickerMenuItem", "Translation Picker"),
+						LOCTEXT("TranslationPickerMenuItemToolTip", "Launch the Translation Picker to Modify Editor Translations"),
+						FSlateIcon(),
+						FUIAction(FExecuteAction::CreateStatic(&FMainFrameTranslationEditorMenu::HandleOpenTranslationPicker))
+						);
 				}
 			}
 			MenuBuilder.EndSection();
diff --git a/Engine/Source/Editor/MainFrame/Private/Menus/TranslationEditorMenu.h b/Engine/Source/Editor/MainFrame/Private/Menus/TranslationEditorMenu.h
--- a/Engine/Source/Editor/MainFrame/Private/Menus/TranslationEditorMenu.h
+++ b/Engine/Source/Editor/MainFrame/Private/Menus/TranslationEditorMenu.h
@@ -166,7 +166,7 @@ public:
 					FText::Format(LOCTEXT("TranslationEditorSubMenuLanguageLabel", "{ArchiveName} Translations"), Arguments),
 					FText::Format(LOCTEXT("TranslationEditorSubMenuLanguageToolTip", "Edit the {ArchiveName} Translations for the {ManifestName} Project"), Arguments),
 					FSlateIcon(),
-					FUIAction(FExecuteAction::CreateStatic(&FMainFrameTranslationEditorMenu::HandleMenuEntryExecute, ManifestFileName, ArchiveFileName))
+					FUIAction(FExecuteAction::CreateStatic(&FMainFrameTranslationEditorMenu::HandleOpenTranslationEditor, ManifestFileName, ArchiveFileName))
 					);
 			}
 		}
@@ -174,10 +174,16 @@ public:
 		MenuBuilder.EndSection();
 	}
 
-	private:
+	// Handles clicking a menu entry.
+	static void HandleOpenTranslationPicker()
+	{
+		FModuleManager::LoadModuleChecked<ITranslationEditor>("TranslationEditor").OpenTranslationPicker();
+	}
+
+private:
 
 	// Handles clicking a menu entry.
-	static void HandleMenuEntryExecute( FString ProjectName, FString Language)
+	static void HandleOpenTranslationEditor( FString ProjectName, FString Language)
 	{
 		FModuleManager::LoadModuleChecked<ITranslationEditor>("TranslationEditor").OpenTranslationEditor(ProjectName, Language);
 	}
diff --git a/Engine/Source/Editor/TranslationEditor/Private/ITranslationEditor.cpp b/Engine/Source/Editor/TranslationEditor/Private/ITranslationEditor.cpp
--- a/Engine/Source/Editor/TranslationEditor/Private/ITranslationEditor.cpp
+++ b/Engine/Source/Editor/TranslationEditor/Private/ITranslationEditor.cpp
@@ -3,6 +3,7 @@
 
 #include "TranslationEditorPrivatePCH.h"
 #include "ITranslationEditor.h"
+#include "Editor/TranslationEditor/Private/TranslationPickerWidget.h"
 
 /** To keep track of what translation editors are open editing which manifest files */
 TMap<FString, ITranslationEditor*> ITranslationEditor::OpenTranslationEditors = TMap<FString, ITranslationEditor*>();
@@ -24,6 +25,14 @@ void ITranslationEditor::OpenTranslationEditor(const FString& InManifestFile, co
 	}
 }
 
+void ITranslationEditor::OpenTranslationPicker()
+{
+	if (!TranslationPickerManager::IsPickerWindowOpen())
+	{
+		TranslationPickerManager::OpenPickerWindow();
+	}
+}
+
 void ITranslationEditor::RegisterTranslationEditor()
 {
 	ITranslationEditor::OpenTranslationEditors.Add(ArchiveFilePath, this);
diff --git a/Engine/Source/Editor/TranslationEditor/Private/TranslationEditor.cpp b/Engine/Source/Editor/TranslationEditor/Private/TranslationEditor.cpp
--- a/Engine/Source/Editor/TranslationEditor/Private/TranslationEditor.cpp
+++ b/Engine/Source/Editor/TranslationEditor/Private/TranslationEditor.cpp
@@ -796,6 +796,10 @@ void FTranslationEditor::MapActions()
 	ToolkitCommands->MapAction(FTranslationEditorCommands::Get().OpenSearchTab,
 		FExecuteAction::CreateSP(this, &FTranslationEditor::OpenSearchTab_Execute),
 		FCanExecuteAction());
+
+	ToolkitCommands->MapAction(FTranslationEditorCommands::Get().OpenTranslationPicker,
+		FExecuteAction::CreateStatic(&ITranslationEditor::OpenTranslationPicker),
+		FCanExecuteAction());
 }
 
 void FTranslationEditor::ChangeSourceFont()
diff --git a/Engine/Source/Editor/TranslationEditor/Private/TranslationEditorMenu.cpp b/Engine/Source/Editor/TranslationEditor/Private/TranslationEditorMenu.cpp
--- a/Engine/Source/Editor/TranslationEditor/Private/TranslationEditorMenu.cpp
+++ b/Engine/Source/Editor/TranslationEditor/Private/TranslationEditorMenu.cpp
@@ -20,6 +20,7 @@ void FTranslationEditorMenu::FillTranslationMenu( FMenuBuilder& MenuBuilder/*, F
 		MenuBuilder.AddMenuEntry( FTranslationEditorCommands::Get().ExportToPortableObjectFormat );
 		MenuBuilder.AddMenuEntry( FTranslationEditorCommands::Get().ImportFromPortableObjectFormat );
 		MenuBuilder.AddMenuEntry( FTranslationEditorCommands::Get().OpenSearchTab );
+		MenuBuilder.AddMenuEntry( FTranslationEditorCommands::Get().OpenTranslationPicker );
 	}
 	MenuBuilder.EndSection();
 }
@@ -100,6 +101,7 @@ void FTranslationEditorCommands::RegisterCommands()
 	UI_COMMAND( ExportToPortableObjectFormat, "Export to .PO", "Export to Portable Object Format", EUserInterfaceActionType::Button, FInputGesture() );
 	UI_COMMAND( ImportFromPortableObjectFormat, "Import from .PO", "Import from Portable Object Format", EUserInterfaceActionType::Button, FInputGesture() );
 	UI_COMMAND( OpenSearchTab, "Search", "Search Source and Translation Strings", EUserInterfaceActionType::Button, FInputGesture() );
+	UI_COMMAND( OpenTranslationPicker, "Translation Picker", "Open the Translation Picker to Modify Editor Translations", EUserInterfaceActionType::Button, FInputGesture() )
 }
 
 
diff --git a/Engine/Source/Editor/TranslationEditor/Private/TranslationEditorModule.cpp b/Engine/Source/Editor/TranslationEditor/Private/TranslationEditorModule.cpp
--- a/Engine/Source/Editor/TranslationEditor/Private/TranslationEditorModule.cpp
+++ b/Engine/Source/Editor/TranslationEditor/Private/TranslationEditorModule.cpp
@@ -30,6 +30,7 @@ void FTranslationEditorModule::StartupModule()
 void FTranslationEditorModule::ShutdownModule()
 {
 	MenuExtensibilityManager.Reset();
+	TranslationPickerManager::ClosePickerWindow();
 
 #if WITH_UNREAL_DEVELOPER_TOOLS
 	// unregister message log
diff --git a/Engine/Source/Editor/TranslationEditor/Private/TranslationPickerEditWindow.cpp b/Engine/Source/Editor/TranslationEditor/Private/TranslationPickerEditWindow.cpp
--- a/Engine/Source/Editor/TranslationEditor/Private/TranslationPickerEditWindow.cpp
+++ b/Engine/Source/Editor/TranslationEditor/Private/TranslationPickerEditWindow.cpp
@@ -167,6 +167,9 @@ void STranslationPickerEditWidget::Construct(const FArguments& InArgs)
 	TranslationUnit->Source = SourceString != nullptr ? *SourceString : "";
 	TranslationUnit->Translation = TranslationString;
 	TranslationUnit->LocresPath = LocresFullPath;
+	
+	// Can only save if we have all the required information
+	bool bCanSave = NamespaceString != nullptr && SourceString != nullptr && LocresFullPath.Len() > 0;
 
 	// Layout all our data
 	ChildSlot
@@ -237,6 +240,7 @@ void STranslationPickerEditWidget::Construct(const FArguments& InArgs)
 					.HAlign(HAlign_Center)
 					.ContentPadding(FEditorStyle::GetMargin("StandardDialog.ContentPadding"))
 					.OnClicked(this, &STranslationPickerEditWidget::SaveAndPreview)
+					.IsEnabled(bCanSave)
 					[
 						SNew(SHorizontalBox)
 						+ SHorizontalBox::Slot()
@@ -245,7 +249,7 @@ void STranslationPickerEditWidget::Construct(const FArguments& InArgs)
 						.AutoWidth()
 						[
 							SNew(STextBlock)
-							.Text(LOCTEXT("SaveAndPreviewButtonText", "Save and preview"))
+							.Text(bCanSave ? LOCTEXT("SaveAndPreviewButtonText", "Save and preview") : LOCTEXT("SaveAndPreviewButtonDisabledText", "Cannot Save") )
 						]
 					]
 				]
diff --git a/Engine/Source/Editor/TranslationEditor/Private/TranslationPickerFloatingWindow.cpp b/Engine/Source/Editor/TranslationEditor/Private/TranslationPickerFloatingWindow.cpp
--- a/Engine/Source/Editor/TranslationEditor/Private/TranslationPickerFloatingWindow.cpp
+++ b/Engine/Source/Editor/TranslationEditor/Private/TranslationPickerFloatingWindow.cpp
@@ -4,19 +4,14 @@
 #include "TranslationPickerFloatingWindow.h"
 #include "Editor/Documentation/Public/SDocumentationToolTip.h"
 #include "TranslationPickerEditWindow.h"
+#include "Editor/TranslationEditor/Private/TranslationPickerWidget.h"
 
 #define LOCTEXT_NAMESPACE "TranslationPicker"
 
 void STranslationPickerFloatingWindow::Construct(const FArguments& InArgs)
 {
 	ParentWindow = InArgs._ParentWindow;
-
-	WindowContents = SNew(SBox);
-
-	WindowContents->SetContent(
-		SNew(SToolTip)
-		.Text(this, &STranslationPickerFloatingWindow::GetPickerStatusText)
-		);
+	WindowContents = SNew(SToolTip);
 
 	ChildSlot
 	[
@@ -26,38 +21,72 @@ void STranslationPickerFloatingWindow::Construct(const FArguments& InArgs)
 
 void STranslationPickerFloatingWindow::Tick( const FGeometry& AllottedGeometry, const double InCurrentTime, const float InDeltaTime )
 {
-	PickedTexts.Empty();
-	TranslationInfoPreviewText = FText::GetEmpty();
-
 	FWidgetPath Path = FSlateApplication::Get().LocateWindowUnderMouse(FSlateApplication::Get().GetCursorPos(), FSlateApplication::Get().GetInteractiveTopLevelWindows(), true);
 
-	// Search everything under the cursor for any FText we know how to parse
-	for (int32 PathIndex = Path.Widgets.Num() - 1; PathIndex >= 0; PathIndex--)
+	if (Path.IsValid())
 	{
-		// General Widget case
-		TSharedRef<SWidget> PathWidget = Path.Widgets[PathIndex].Widget;
-		FText TextBlockDescription = GetTextDescription(PathWidget);
-
-		if (!TextBlockDescription.IsEmpty())
+		// If the path of widgets we're hovering over changed since last time (or if this is the first tick and LastTickHoveringWidgetPath hasn't been set yet)
+		if (!LastTickHoveringWidgetPath.IsValid() || LastTickHoveringWidgetPath.ToWidgetPath().ToString() != Path.ToString())
 		{
-			TranslationInfoPreviewText = FText::Format(LOCTEXT("TextBlockLocalizationDescription", "{0}\n\n\n\n\nTextBlock info:\n{1}"),
-				TranslationInfoPreviewText, TextBlockDescription);
-		}
+			// Clear all previous text and widgets
+			PickedTexts.Empty();
+			WindowContents->SetContentWidget(SNew(SBox));
 
-		// Tooltip case
-		TSharedPtr<IToolTip> Tooltip = PathWidget->GetToolTip();
-		FText TooltipDescription = FText::GetEmpty();
+			// Search everything under the cursor for any FText we know how to parse
+			for (int32 PathIndex = Path.Widgets.Num() - 1; PathIndex >= 0; PathIndex--)
+			{
+				// General Widget case
+				TSharedRef<SWidget> PathWidget = Path.Widgets[PathIndex].Widget;
+				GetTextFromWidget(PathWidget);
+
+				// Tooltip case
+				TSharedPtr<IToolTip> Tooltip = PathWidget->GetToolTip();
+				//FText TooltipDescription = FText::GetEmpty();
+
+				if (Tooltip.IsValid() && !Tooltip->IsEmpty())
+				{
+					GetTextFromWidget(Tooltip->AsWidget());
+				}
+			}
 
-		if (Tooltip.IsValid() && !Tooltip->IsEmpty())
-		{
-			TooltipDescription = GetTextDescription(Tooltip->AsWidget());
-			if (!TooltipDescription.IsEmpty())
+			auto TextsBox = SNew(SVerticalBox);
+
+			// Add a new Translation Picker Edit Widget for each picked text
+			for (FText PickedText : PickedTexts)
 			{
-				TranslationInfoPreviewText = FText::Format(LOCTEXT("TooltipLocalizationDescription", "{0}\n\n\n\n\nTooltip Info:\n{1}"),
-					TranslationInfoPreviewText, TooltipDescription);
+				TextsBox->AddSlot()
+					.FillHeight(1)
+					.Padding(FMargin(5))
+					[
+						SNew(SBorder)
+						[
+							SNew(STranslationPickerEditWidget)
+							.PickedText(PickedText)
+						]
+					];
 			}
+
+			FText Instructions = FText::GetEmpty();
+			if (PickedTexts.Num() >= 1)
+			{
+				Instructions = LOCTEXT("TranslationPickerEscToEdit", "Press Esc to edit this translation");
+			}
+			else
+			{
+				Instructions = LOCTEXT("TranslationPickerHoverToViewEditEscToQuit", "Hover over text to view/edit translation, or press Esc to quit");
+			}
+
+			TextsBox->AddSlot()
+				.AutoHeight()
+				.Padding(FMargin(5))
+				[
+					SNew(STextBlock)
+					.Text(Instructions)
+					.Justification(ETextJustify::Center)
+				];
+
+			WindowContents->SetContentWidget(TextsBox);
 		}
-		
 	}
 
 	// kind of a hack, but we need to maintain keyboard focus otherwise we wont get our keypress to 'pick'
@@ -67,19 +96,19 @@ void STranslationPickerFloatingWindow::Tick( const FGeometry& AllottedGeometry,
 		// also kind of a hack, but this is the only way at the moment to get a 'cursor decorator' without using the drag-drop code path
 		ParentWindow.Pin()->MoveWindowTo(FSlateApplication::Get().GetCursorPos() + FSlateApplication::Get().GetCursorSize());
 	}
+
+	LastTickHoveringWidgetPath = FWeakWidgetPath(Path);
 }
 
-FText STranslationPickerFloatingWindow::GetTextDescription(TSharedRef<SWidget> Widget)
+FText STranslationPickerFloatingWindow::GetTextFromWidget(TSharedRef<SWidget> Widget)
 {
 	FText OriginalText = FText::GetEmpty();
-	FText FormattedText = FText::GetEmpty();
 
 	// Parse STextBlocks
 	STextBlock* TextBlock = (STextBlock*)&Widget.Get();
 	if ((Widget->GetTypeAsString() == "STextBlock" && TextBlock))
 	{
 		OriginalText = TextBlock->GetText();
-		FormattedText = FormatFTextInfo(OriginalText);
 	}
 
 	// Parse SToolTips
@@ -88,15 +117,13 @@ FText STranslationPickerFloatingWindow::GetTextDescription(TSharedRef<SWidget> W
 		SToolTip* ToolTip = ((SToolTip*)&Widget.Get());
 		if (ToolTip != nullptr)
 		{
-			FormattedText = GetTextDescription(ToolTip->GetContentWidget());
-			if (FormattedText.IsEmpty())
+			OriginalText = GetTextFromWidget(ToolTip->GetContentWidget());
+			if (OriginalText.IsEmpty())
 			{
 				OriginalText = ToolTip->GetTextTooltip();
-				FormattedText = FormatFTextInfo(OriginalText);
 			}
 		}
 	}
-
 	// Parse SDocumentationToolTips
 	else if (Widget->GetTypeAsString() == "SDocumentationToolTip")
 	{
@@ -104,46 +131,25 @@ FText STranslationPickerFloatingWindow::GetTextDescription(TSharedRef<SWidget> W
 		if (DocumentationToolTip != nullptr)
 		{
 			OriginalText = DocumentationToolTip->GetTextTooltip();
-			FormattedText = FormatFTextInfo(OriginalText);
 		}
 	}
 	
-	if (!OriginalText.IsEmpty())
-	{
-		PickedTexts.Add(OriginalText);
-	}
-
-	return FormattedText;
-}
-
-FText STranslationPickerFloatingWindow::FormatFTextInfo(FText TextToFormat)
-{
-	FText OutText = FText::GetEmpty();
-
-	if (FTextInspector::ShouldGatherForLocalization(TextToFormat))
+	// Don't show the same text twice
+	bool bAlreadyPicked = false;
+	for (FText PickedText : PickedTexts)
 	{
-		const FString* TextNamespace = FTextInspector::GetNamespace(TextToFormat);
-		const FString* TextKey = FTextInspector::GetKey(TextToFormat);
-		const FString* TextSource = FTextInspector::GetSourceString(TextToFormat);
-		const FString& TextTranslation = FTextInspector::GetDisplayString(TextToFormat);
-
-		TSharedPtr<FString, ESPMode::ThreadSafe> TextTableName = TSharedPtr<FString, ESPMode::ThreadSafe>(nullptr);
-		if (TextNamespace && TextKey)
+		if (OriginalText.EqualTo(PickedText))
 		{
-			TextTableName = FTextLocalizationManager::Get().GetTableName(*TextNamespace, *TextKey);
+			bAlreadyPicked = true;
 		}
-
-		FText Namespace = TextNamespace != nullptr ? FText::FromString(*TextNamespace) : FText::GetEmpty();
-		FText Key = TextKey != nullptr ? FText::FromString(*TextKey) : FText::GetEmpty();
-		FText Source = TextSource != nullptr ? FText::FromString(*TextSource) : FText::GetEmpty();
-		FText TableName = TextTableName.IsValid() ? FText::FromString(*(TextTableName.Get())) : FText::GetEmpty();
-		FText Translation = FText::FromString(TextTranslation);
-
-		OutText = FText::Format(LOCTEXT("LocalizationStringDescription", "Text Namespace: {0}\nText Key: {1}\nText Source: {2}\nTable Name: {3}"), Namespace, Key, Source, TableName);
-		OutText = FText::Format(LOCTEXT("LocalizationStringDescription2", "{0}\nTranslation: {1}"), OutText, Translation);
+	}
+	
+	if (!OriginalText.IsEmpty() && !bAlreadyPicked)
+	{
+		PickedTexts.Add(OriginalText);
 	}
 
-	return OutText; 
+	return OriginalText;
 }
 
 FReply STranslationPickerFloatingWindow::OnKeyDown(const FGeometry& MyGeometry, const FKeyEvent& InKeyEvent)
@@ -175,9 +181,7 @@ FReply STranslationPickerFloatingWindow::OnKeyDown(const FGeometry& MyGeometry,
 			FSlateApplication::Get().AddWindow(NewWindow);
 		}
 
-		Close();
-
-		TranslationInfoPreviewText = FText::GetEmpty();
+		TranslationPickerManager::ClosePickerWindow();
 		
 		return FReply::Handled();
 	}
@@ -185,16 +189,5 @@ FReply STranslationPickerFloatingWindow::OnKeyDown(const FGeometry& MyGeometry,
 	return FReply::Unhandled();
 }
 
-FReply STranslationPickerFloatingWindow::Close()
-{
-	if (ParentWindow.IsValid())
-	{
-		FSlateApplication::Get().RequestDestroyWindow(ParentWindow.Pin().ToSharedRef());
-		ParentWindow.Reset();
-	}
-
-	return FReply::Handled();
-}
-
 
 #undef LOCTEXT_NAMESPACE
\ No newline at end of file
diff --git a/Engine/Source/Editor/TranslationEditor/Private/TranslationPickerFloatingWindow.h b/Engine/Source/Editor/TranslationEditor/Private/TranslationPickerFloatingWindow.h
--- a/Engine/Source/Editor/TranslationEditor/Private/TranslationPickerFloatingWindow.h
+++ b/Engine/Source/Editor/TranslationEditor/Private/TranslationPickerFloatingWindow.h
@@ -20,23 +20,11 @@ private:
 	virtual void Tick( const FGeometry& AllottedGeometry, const double InCurrentTime, const float InDeltaTime ) override;
 
 	/** Pull the FText reference out of an SWidget */
-	FText GetTextDescription(TSharedRef<SWidget> Widget);
-
-	/** Format an FText for display */
-	FText FormatFTextInfo(FText TextToFormat);
-
-	/** This is what this window displays */
-	FText GetPickerStatusText() const
-	{
-		return FText::Format(LOCTEXT("TootipHint", "{0}\n\n(Esc to pick)"), TranslationInfoPreviewText);
-	}
+	FText GetTextFromWidget(TSharedRef<SWidget> Widget);
 
 	/** Handle key presses */
 	virtual FReply OnKeyDown(const FGeometry& MyGeometry, const FKeyEvent& InKeyEvent) override;
 
-	/** Close this window */
-	FReply Close();
-
 	/** We need to support keyboard focus to process the 'Esc' key */
 	virtual bool SupportsKeyboardFocus() const override
 	{
@@ -47,13 +35,15 @@ private:
 	TWeakPtr<SWindow> ParentWindow;
 
 	/** Contents of the window */
-	TSharedPtr<SBox> WindowContents;
+	TSharedPtr<SToolTip> WindowContents;
 
 	/** The FTexts that we have found under the cursor */
 	TArray<FText> PickedTexts;
 
-	/** Preview text for the translation */
-	FText TranslationInfoPreviewText;
+	/**
+	* The path widgets we were hovering over last tick
+	*/
+	FWeakWidgetPath LastTickHoveringWidgetPath;
 };
 
 #undef LOCTEXT_NAMESPACE
\ No newline at end of file
diff --git a/Engine/Source/Editor/TranslationEditor/Private/TranslationPickerWidget.cpp b/Engine/Source/Editor/TranslationEditor/Private/TranslationPickerWidget.cpp
--- a/Engine/Source/Editor/TranslationEditor/Private/TranslationPickerWidget.cpp
+++ b/Engine/Source/Editor/TranslationEditor/Private/TranslationPickerWidget.cpp
@@ -5,16 +5,6 @@
 
 #define LOCTEXT_NAMESPACE "TranslationPicker"
 
-STranslationWidgetPicker::~STranslationWidgetPicker()
-{
-	// kill the picker window as well if this widget is going away - that way we dont get dangling refs to the property
-	if (PickerWindow.IsValid() && FSlateApplication::IsInitialized())
-	{
-		FSlateApplication::Get().RequestDestroyWindow(PickerWindow.Pin().ToSharedRef());
-		PickerWindow.Reset();
-		PickerWindowWidget.Reset();
-	}
-}
 
 void STranslationWidgetPicker::Construct(const FArguments& InArgs)
 {
@@ -85,13 +75,33 @@ void STranslationWidgetPicker::Construct(const FArguments& InArgs)
 
 
 				// Bind the check box's "checked" state to our user interface action
-				.IsChecked(this, &STranslationWidgetPicker::OnIsChecked)
+				.IsChecked(this, &STranslationWidgetPicker::IsChecked)
 
 				.Padding(FEditorStyle::Get().GetMargin(CheckboxStyle))
 		];
 }
 
-FReply STranslationWidgetPicker::OnClicked()
+ECheckBoxState STranslationWidgetPicker::IsChecked() const
+{
+	return TranslationPickerManager::IsPickerWindowOpen() ? ECheckBoxState::Checked : ECheckBoxState::Unchecked;
+}
+
+void STranslationWidgetPicker::OnCheckStateChanged(const ECheckBoxState NewCheckedState)
+{
+	if (TranslationPickerManager::IsPickerWindowOpen())
+	{
+		TranslationPickerManager::ClosePickerWindow();
+	}
+	else
+	{
+		TranslationPickerManager::OpenPickerWindow();
+	}
+}
+
+TSharedPtr<SWindow> TranslationPickerManager::PickerWindow = TSharedPtr<SWindow>();
+TSharedPtr<STranslationPickerFloatingWindow> TranslationPickerManager::PickerWindowWidget = TSharedPtr<STranslationPickerFloatingWindow>();
+
+bool TranslationPickerManager::OpenPickerWindow()
 {
 	// Not picking previously, launch a picker window
 	if (!PickerWindow.IsValid())
@@ -106,43 +116,29 @@ FReply STranslationWidgetPicker::OnClicked()
 			);
 
 		TSharedPtr<SWindow> RootWindow = FGlobalTabmanager::Get()->GetRootWindow();
-		check(RootWindow.IsValid());
-		FSlateApplication::Get().AddWindowAsNativeChild(NewWindow, RootWindow.ToSharedRef());
-	}
-	// Was already picking previously, destroy picker window
-	else
-	{
-		FSlateApplication::Get().RequestDestroyWindow(PickerWindow.Pin().ToSharedRef());
-		PickerWindow.Pin().Reset();
-	}
+		if (RootWindow.IsValid())
+		{
+			FSlateApplication::Get().AddWindowAsNativeChild(NewWindow, RootWindow.ToSharedRef());
+		}
+		else
+		{
+			FSlateApplication::Get().AddWindow(NewWindow);
+		}
 
-	return FReply::Handled();
-}
-
-bool STranslationWidgetPicker::IsPicking() const
-{
-	if (PickerWindowWidget.IsValid())
-	{
 		return true;
 	}
 
 	return false;
 }
 
-ECheckBoxState STranslationWidgetPicker::OnIsChecked() const
-{
-	return IsPicking() ? ECheckBoxState::Checked : ECheckBoxState::Unchecked;
-}
-
-void STranslationWidgetPicker::OnCheckStateChanged(const ECheckBoxState NewCheckedState)
+void TranslationPickerManager::ClosePickerWindow()
 {
-	OnClicked();
+	if (FSlateApplication::IsInitialized())
+	{
+		FSlateApplication::Get().RequestDestroyWindow(PickerWindow.ToSharedRef());
+	}
+	PickerWindow.Reset(); 
+	PickerWindowWidget.Reset();
 }
 
-
-
-
-#undef LOCTEXT_NAMESPACE
-
-
-
+#undef LOCTEXT_NAMESPACE
\ No newline at end of file
diff --git a/Engine/Source/Editor/TranslationEditor/Private/TranslationPickerWidget.h b/Engine/Source/Editor/TranslationEditor/Private/TranslationPickerWidget.h
--- a/Engine/Source/Editor/TranslationEditor/Private/TranslationPickerWidget.h
+++ b/Engine/Source/Editor/TranslationEditor/Private/TranslationPickerWidget.h
@@ -6,6 +6,19 @@
 
 #define LOCTEXT_NAMESPACE "TranslationPicker"
 
+class TranslationPickerManager
+{
+public:
+	static TSharedPtr<SWindow> PickerWindow;
+	static TSharedPtr<STranslationPickerFloatingWindow> PickerWindowWidget;
+
+	static bool IsPickerWindowOpen() { return PickerWindow.IsValid(); }
+
+	static bool OpenPickerWindow();
+
+	static void ClosePickerWindow();
+};
+
 /** Widget used to launch a 'picking' session */
 class STranslationWidgetPicker : public SCompoundWidget
 {
@@ -14,27 +27,16 @@ public:
 
 	SLATE_END_ARGS()
 
-	~STranslationWidgetPicker();
-
 	void Construct(const FArguments& InArgs);
 
 private:
-	FReply OnClicked();
-
-	bool IsPicking() const;
-
-	/** Picker window widget */
-	TWeakPtr<STranslationPickerFloatingWindow> PickerWindowWidget;
-
-	/** Picker window */
-	TWeakPtr<SWindow> PickerWindow;
 
 	/**
 	* Called by slate to determine if this button should appear checked
 	*
 	* @return ECheckBoxState::Checked if it should be checked, ECheckBoxState::Unchecked if not.
 	*/
-	ECheckBoxState OnIsChecked() const;
+	ECheckBoxState IsChecked() const;
 
 	/**
 	* Called by Slate when this tool bar check box button is toggled
diff --git a/Engine/Source/Editor/TranslationEditor/Public/ITranslationEditor.h b/Engine/Source/Editor/TranslationEditor/Public/ITranslationEditor.h
--- a/Engine/Source/Editor/TranslationEditor/Public/ITranslationEditor.h
+++ b/Engine/Source/Editor/TranslationEditor/Public/ITranslationEditor.h
@@ -5,6 +5,7 @@
 #include "Toolkits/IToolkit.h"
 #include "Toolkits/AssetEditorToolkit.h"
 #include "TranslationEditorModule.h"
+#include "../Private/TranslationPickerWidget.h"
 
 
 /** Translation Editor public interface */
@@ -15,6 +16,8 @@ public:
 
 	TRANSLATIONEDITOR_API static void OpenTranslationEditor(const FString& InManifestFile, const FString& InArchiveFile);
 
+	TRANSLATIONEDITOR_API static void OpenTranslationPicker();
+
 	ITranslationEditor(const FString& InManifestFile, const FString& InArchiveFile)
 		: ManifestFilePath(InManifestFile)
 		, ArchiveFilePath(InArchiveFile)
diff --git a/Engine/Source/Editor/TranslationEditor/Public/TranslationEditorMenu.h b/Engine/Source/Editor/TranslationEditor/Public/TranslationEditorMenu.h
--- a/Engine/Source/Editor/TranslationEditor/Public/TranslationEditorMenu.h
+++ b/Engine/Source/Editor/TranslationEditor/Public/TranslationEditorMenu.h
@@ -45,6 +45,9 @@ public:
 	/** Open the tab for searching */
 	TSharedPtr<FUICommandInfo> OpenSearchTab;
 
+	/** Open the translation picker */
+	TSharedPtr<FUICommandInfo> OpenTranslationPicker;
+
 	/** Initialize commands */
 	virtual void RegisterCommands() override;
 };
