Translation Picker improvements:
- Make the floating preview not overlap tooltips by putting it on the left side of the cursor
- Implemented "save all" button
- Other small fixes

diff --git a/Engine/Source/Editor/TranslationEditor/Private/TranslationPickerEditWindow.cpp b/Engine/Source/Editor/TranslationEditor/Private/TranslationPickerEditWindow.cpp
--- a/Engine/Source/Editor/TranslationEditor/Private/TranslationPickerEditWindow.cpp
+++ b/Engine/Source/Editor/TranslationEditor/Private/TranslationPickerEditWindow.cpp
@@ -16,20 +16,36 @@ void STranslationPickerEditWindow::Construct(const FArguments& InArgs)
 
 	auto TextsBox = SNew(SVerticalBox);
 
+	// Display name of the current language
+	TextsBox->AddSlot()
+		.AutoHeight()
+		.Padding(FMargin(5))
+		[
+			SNew(STextBlock)
+			.Text(FText::FromString(FInternationalization::Get().GetCurrentCulture()->GetDisplayName()))
+			.Justification(ETextJustify::Center)
+		];
+
 	// Add a new Translation Picker Edit Widget for each picked text
 	for (FText PickedText : PickedTexts)
 	{
 		TSharedPtr<SEditableTextBox> TextBox;
 		int32 DefaultPadding = 0.0f;
 
+		TSharedRef<STranslationPickerEditWidget> NewEditWidget = 
+			SNew(STranslationPickerEditWidget)
+			.PickedText(PickedText)
+			.bAllowEditing(true);
+
+		EditWidgets.Add(NewEditWidget);
+
 		TextsBox->AddSlot()
 			.FillHeight(1)
 			.Padding(FMargin(5))
 			[
 				SNew(SBorder)
 				[
-					SNew(STranslationPickerEditWidget)
-					.PickedText(PickedText)
+					NewEditWidget
 				]
 			];
 	}
@@ -70,7 +86,7 @@ void STranslationPickerEditWindow::Construct(const FArguments& InArgs)
 							SNew(SButton)
 							.HAlign(HAlign_Center)
 							.ContentPadding(FEditorStyle::GetMargin("StandardDialog.ContentPadding"))
-							.OnClicked(this, &STranslationPickerEditWindow::Close)
+							.OnClicked(this, &STranslationPickerEditWindow::SaveAllAndClose)
 							[
 								SNew(SHorizontalBox)
 								+ SHorizontalBox::Slot()
@@ -126,9 +142,35 @@ FReply STranslationPickerEditWindow::Close()
 	return FReply::Handled();
 }
 
+FReply STranslationPickerEditWindow::SaveAllAndClose()
+{
+	
+	TArray<UTranslationUnit*> TempArray;
+
+	for (TSharedRef<STranslationPickerEditWidget> EditWidget : EditWidgets)
+	{
+		UTranslationUnit* TranslationUnit = EditWidget->GetTranslationUnitWithAnyChanges();
+		if (TranslationUnit != nullptr)
+		{
+			TempArray.Add(TranslationUnit);
+		}
+	}
+
+	if (TempArray.Num() > 0)
+	{
+		// Save the data via translation data manager
+		FTranslationDataManager::SaveSelectedTranslations(TempArray);
+	}
+
+	Close();
+
+	return FReply::Handled();
+}
+
 void STranslationPickerEditWidget::Construct(const FArguments& InArgs)
 {
 	PickedText = InArgs._PickedText;
+	bAllowEditing = InArgs._bAllowEditing;
 	int32 DefaultPadding = 0.0f;
 
 	// Get all the data we need and format it properly
@@ -201,6 +243,7 @@ void STranslationPickerEditWidget::Construct(const FArguments& InArgs)
 					+ SHorizontalBox::Slot()
 					[
 						SAssignNew(TextBox, SEditableTextBox)
+						.IsEnabled(bAllowEditing && bCanSave)
 						.Text(Translation)
 						.HintText(LOCTEXT("TranslationEditTextBox_HintText", "Enter/edit translation here."))
 					]
@@ -241,6 +284,7 @@ void STranslationPickerEditWidget::Construct(const FArguments& InArgs)
 					.ContentPadding(FEditorStyle::GetMargin("StandardDialog.ContentPadding"))
 					.OnClicked(this, &STranslationPickerEditWidget::SaveAndPreview)
 					.IsEnabled(bCanSave)
+					.Visibility(bAllowEditing ? EVisibility::Visible : EVisibility::Collapsed)
 					[
 						SNew(SHorizontalBox)
 						+ SHorizontalBox::Slot()
@@ -270,4 +314,17 @@ FReply STranslationPickerEditWidget::SaveAndPreview()
 	return FReply::Handled();
 }
 
+UTranslationUnit* STranslationPickerEditWidget::GetTranslationUnitWithAnyChanges()
+{
+	if (TranslationUnit)
+	{
+		// Update translation string from entered text
+		TranslationUnit->Translation = TextBox->GetText().ToString();
+
+		return TranslationUnit;
+	}
+
+	return nullptr;
+}
+
 #undef LOCTEXT_NAMESPACE
\ No newline at end of file
diff --git a/Engine/Source/Editor/TranslationEditor/Private/TranslationPickerEditWindow.h b/Engine/Source/Editor/TranslationEditor/Private/TranslationPickerEditWindow.h
--- a/Engine/Source/Editor/TranslationEditor/Private/TranslationPickerEditWindow.h
+++ b/Engine/Source/Editor/TranslationEditor/Private/TranslationPickerEditWindow.h
@@ -12,10 +12,15 @@ class STranslationPickerEditWidget : public SCompoundWidget
 
 	SLATE_ARGUMENT(FText, PickedText)
 
+	SLATE_ARGUMENT(bool, bAllowEditing)
+
 	SLATE_END_ARGS()
 
 	void Construct(const FArguments& InArgs);
 
+	/** Return the translation unit for this text, with any modifications */
+	UTranslationUnit* GetTranslationUnitWithAnyChanges();
+
 private:
 
 	virtual bool SupportsKeyboardFocus() const override
@@ -33,6 +38,9 @@ private:
 
 	/** The text box for entering/modifying a translation */
 	TSharedPtr<SEditableTextBox> TextBox;
+
+	/** Whether or not to show the save button*/
+	bool bAllowEditing;
 };
 
 /** Translation picker edit window to allow you to translate selected FTexts in place */
@@ -68,6 +76,12 @@ private:
 
 	/** The FTexts that we have found under the cursor */
 	TArray<FText> PickedTexts;
+
+	/** All of our current edit widgets */
+	TArray<TSharedRef<STranslationPickerEditWidget>> EditWidgets;
+
+	/** Save all translations and close */
+	FReply SaveAllAndClose();
 };
 
 #undef LOCTEXT_NAMESPACE
\ No newline at end of file
diff --git a/Engine/Source/Editor/TranslationEditor/Private/TranslationPickerFloatingWindow.cpp b/Engine/Source/Editor/TranslationEditor/Private/TranslationPickerFloatingWindow.cpp
--- a/Engine/Source/Editor/TranslationEditor/Private/TranslationPickerFloatingWindow.cpp
+++ b/Engine/Source/Editor/TranslationEditor/Private/TranslationPickerFloatingWindow.cpp
@@ -51,6 +51,16 @@ void STranslationPickerFloatingWindow::Tick( const FGeometry& AllottedGeometry,
 
 			auto TextsBox = SNew(SVerticalBox);
 
+			// Display name of the current language
+			TextsBox->AddSlot()
+				.AutoHeight()
+				.Padding(FMargin(5))
+				[
+					SNew(STextBlock)
+					.Text(FText::FromString(FInternationalization::Get().GetCurrentCulture()->GetDisplayName()))
+					.Justification(ETextJustify::Center)
+				];
+
 			// Add a new Translation Picker Edit Widget for each picked text
 			for (FText PickedText : PickedTexts)
 			{
@@ -62,6 +72,7 @@ void STranslationPickerFloatingWindow::Tick( const FGeometry& AllottedGeometry,
 						[
 							SNew(STranslationPickerEditWidget)
 							.PickedText(PickedText)
+							.bAllowEditing(false)
 						]
 					];
 			}
@@ -93,8 +104,16 @@ void STranslationPickerFloatingWindow::Tick( const FGeometry& AllottedGeometry,
 	FSlateApplication::Get().SetKeyboardFocus(SharedThis(this), EKeyboardFocusCause::SetDirectly);
 	if (ParentWindow.IsValid())
 	{
+		FVector2D WindowSize = ParentWindow.Pin()->GetDesiredSize();
+		FVector2D DesiredPosition = FSlateApplication::Get().GetCursorPos();
+		DesiredPosition.X -= FSlateApplication::Get().GetCursorSize().X;
+		DesiredPosition.Y += FSlateApplication::Get().GetCursorSize().Y;
+
+		// Move to opposite side of the cursor than the tool tip, so they don't overlaps
+		DesiredPosition.X -= WindowSize.X;
+
 		// also kind of a hack, but this is the only way at the moment to get a 'cursor decorator' without using the drag-drop code path
-		ParentWindow.Pin()->MoveWindowTo(FSlateApplication::Get().GetCursorPos() + FSlateApplication::Get().GetCursorSize());
+		ParentWindow.Pin()->MoveWindowTo(DesiredPosition);
 	}
 
 	LastTickHoveringWidgetPath = FWeakWidgetPath(Path);
@@ -156,29 +175,41 @@ FReply STranslationPickerFloatingWindow::OnKeyDown(const FGeometry& MyGeometry,
 {
 	if (InKeyEvent.GetKey() == EKeys::Escape)
 	{
+		if (PickedTexts.Num() > 0)
 		// Open a different window to allow editing of the translation
-		TSharedRef<SWindow> NewWindow = SNew(SWindow)
-		.Title(LOCTEXT("TranslationPickerEditWindowTitle", "Edit Translation(s)"))
-		.SupportsMaximize(false)
-		.SupportsMinimize(false)
-		.CreateTitleBar(true)
-		.SizingRule(ESizingRule::Autosized);
-		NewWindow->MoveWindowTo(FSlateApplication::Get().GetCursorPos());
-
-		NewWindow->SetContent(
-			SNew(STranslationPickerEditWindow)
-			.ParentWindow(NewWindow)
-			.PickedTexts(PickedTexts)
-			);
-
-		TSharedPtr<SWindow> RootWindow = FGlobalTabmanager::Get()->GetRootWindow();
-		if (RootWindow.IsValid())
 		{
-			FSlateApplication::Get().AddWindowAsNativeChild(NewWindow, RootWindow.ToSharedRef());
-		}
-		else
-		{
-			FSlateApplication::Get().AddWindow(NewWindow);
+			TSharedRef<SWindow> NewWindow = SNew(SWindow)
+				.Title(LOCTEXT("TranslationPickerEditWindowTitle", "Edit Translation(s)"))
+				.SupportsMaximize(false)
+				.SupportsMinimize(false)
+				.CreateTitleBar(true)
+				.SizingRule(ESizingRule::Autosized);
+
+			NewWindow->SetContent(
+				SNew(STranslationPickerEditWindow)
+				.ParentWindow(NewWindow)
+				.PickedTexts(PickedTexts)
+				);
+
+			TSharedPtr<SWindow> RootWindow = FGlobalTabmanager::Get()->GetRootWindow();
+			if (RootWindow.IsValid())
+			{
+				FSlateApplication::Get().AddWindowAsNativeChild(NewWindow, RootWindow.ToSharedRef());
+			}
+			else
+			{
+				FSlateApplication::Get().AddWindow(NewWindow);
+			}
+
+			FVector2D WindowSize = ParentWindow.Pin()->GetDesiredSize();
+			FVector2D DesiredPosition = FSlateApplication::Get().GetCursorPos();
+			DesiredPosition.X -= FSlateApplication::Get().GetCursorSize().X;
+			DesiredPosition.Y += FSlateApplication::Get().GetCursorSize().Y;
+
+			// Open this new Edit window in the same position
+			DesiredPosition.X -= WindowSize.X;
+
+			NewWindow->MoveWindowTo(DesiredPosition);
 		}
 
 		TranslationPickerManager::ClosePickerWindow();
