Translation Editor:
- When saving, if source control is disabled, try write the file anyway.
- Remove the blank icons from the tabs, at least until we get proper icons.

diff --git a/Engine/Source/Editor/TranslationEditor/Private/TranslationDataManager.cpp b/Engine/Source/Editor/TranslationEditor/Private/TranslationDataManager.cpp
--- a/Engine/Source/Editor/TranslationEditor/Private/TranslationDataManager.cpp
+++ b/Engine/Source/Editor/TranslationEditor/Private/TranslationDataManager.cpp
@@ -239,25 +239,29 @@ bool FTranslationDataManager::WriteJSONToTextFile(TSharedRef<FJsonObject>& Outpu
 	// If the user specified a reference file - write the entries read from code to a ref file
 	if ( !Filename.IsEmpty() )
 	{
-		// Already checked out?
-		if(CheckedOutFiles.Contains(Filename))
+		// If source control is enabled, try to check out the file. Otherwise just try to write it
+		if (ISourceControlModule::Get().IsEnabled())
 		{
-			bPreviouslyCheckedOut = true;
-		}
-		else if( !SourceControlHelpers::CheckOutFile(Filename) )
-		{
-			FFormatNamedArguments Arguments;
-			Arguments.Add( TEXT("Filename"), FText::FromString(Filename) );
-			// Use Source Control Message Log here because there might be other useful information in that log for the user.
-			FMessageLog SourceControlMessageLog("SourceControl");
-			SourceControlMessageLog.Error(FText::Format(LOCTEXT("CheckoutFailed", "Check out of file '{Filename}' failed."), Arguments));
-			SourceControlMessageLog.Notify(LOCTEXT("TranslationArchiveCheckoutFailed", "Failed to Check Out Translation Archive!"));
-			SourceControlMessageLog.Open(EMessageSeverity::Error);
-			CheckoutAndSaveWasSuccessful = false;
-		}
-		else
-		{
-			CheckedOutFiles.Add(Filename);
+			// Already checked out?
+			if (CheckedOutFiles.Contains(Filename))
+			{
+				bPreviouslyCheckedOut = true;
+			}
+			else if (!SourceControlHelpers::CheckOutFile(Filename))
+			{
+				FFormatNamedArguments Arguments;
+				Arguments.Add(TEXT("Filename"), FText::FromString(Filename));
+				// Use Source Control Message Log here because there might be other useful information in that log for the user.
+				FMessageLog SourceControlMessageLog("SourceControl");
+				SourceControlMessageLog.Error(FText::Format(LOCTEXT("CheckoutFailed", "Check out of file '{Filename}' failed."), Arguments));
+				SourceControlMessageLog.Notify(LOCTEXT("TranslationArchiveCheckoutFailed", "Failed to Check Out Translation Archive!"));
+				SourceControlMessageLog.Open(EMessageSeverity::Error);
+				CheckoutAndSaveWasSuccessful = false;
+			}
+			else
+			{
+				CheckedOutFiles.Add(Filename);
+			}
 		}
 
 		if( CheckoutAndSaveWasSuccessful )
diff --git a/Engine/Source/Editor/TranslationEditor/Private/TranslationEditor.cpp b/Engine/Source/Editor/TranslationEditor/Private/TranslationEditor.cpp
--- a/Engine/Source/Editor/TranslationEditor/Private/TranslationEditor.cpp
+++ b/Engine/Source/Editor/TranslationEditor/Private/TranslationEditor.cpp
@@ -280,7 +280,7 @@ TSharedRef<SDockTab> FTranslationEditor::SpawnTab_Untranslated( const FSpawnTabA
 	TSharedRef<SWidget> PropertyTableWidget = UntranslatedPropertyTableWidgetHandle->GetWidget();
 
 	TSharedRef<SDockTab> NewDockTab = SNew(SDockTab)
-		.Icon( FEditorStyle::GetBrush("TranslationEditor.Tabs.Properties") )
+		//.Icon( FEditorStyle::GetBrush("TranslationEditor.Tabs.Properties") )
 		.Label( LOCTEXT("UntranslatedTabTitle", "Untranslated") )
 		.TabColorScale( GetTabColorScale() )
 		[
@@ -345,7 +345,7 @@ TSharedRef<SDockTab> FTranslationEditor::SpawnTab_Review( const FSpawnTabArgs& A
 	TSharedRef<SWidget> PropertyTableWidget = ReviewPropertyTableWidgetHandle->GetWidget();
 
 	TSharedRef<SDockTab> NewDockTab = SNew(SDockTab)
-		.Icon( FEditorStyle::GetBrush("TranslationEditor.Tabs.Properties") )
+		//.Icon( FEditorStyle::GetBrush("TranslationEditor.Tabs.Properties") )
 		.Label( LOCTEXT("ReviewTabTitle", "Needs Review") )
 		.TabColorScale( GetTabColorScale() )
 		[
@@ -403,7 +403,7 @@ TSharedRef<SDockTab> FTranslationEditor::SpawnTab_Completed( const FSpawnTabArgs
 	TSharedRef<SWidget> PropertyTableWidget = CompletedPropertyTableWidgetHandle->GetWidget();
 
 	TSharedRef<SDockTab> NewDockTab = SNew(SDockTab)
-		.Icon( FEditorStyle::GetBrush("TranslationEditor.Tabs.Properties") )
+		//.Icon( FEditorStyle::GetBrush("TranslationEditor.Tabs.Properties") )
 		.Label( LOCTEXT("CompletedTabTitle", "Completed") )
 		.TabColorScale( GetTabColorScale() )
 		[
@@ -541,7 +541,7 @@ TSharedRef<SDockTab> FTranslationEditor::SpawnTab_ChangedOnImport(const FSpawnTa
 	TSharedRef<SWidget> PropertyTableWidget = SearchPropertyTableWidgetHandle->GetWidget();
 
 	TSharedRef<SDockTab> NewDockTab = SNew(SDockTab)
-		.Icon(FEditorStyle::GetBrush("TranslationEditor.Tabs.Properties"))
+		//.Icon(FEditorStyle::GetBrush("TranslationEditor.Tabs.Properties"))
 		.Label(LOCTEXT("ChangedOnImportTabTitle", "Changed on Import"))
 		.TabColorScale(GetTabColorScale())
 		[
@@ -563,7 +563,7 @@ TSharedRef<SDockTab> FTranslationEditor::SpawnTab_Preview( const FSpawnTabArgs&
 	check( Args.GetTabId().TabType == PreviewTabId );
 
 	TSharedRef<SDockTab> NewDockTab = SNew(SDockTab)
-		.Icon( FEditorStyle::GetBrush("TranslationEditor.Tabs.Properties") )
+		//.Icon( FEditorStyle::GetBrush("TranslationEditor.Tabs.Properties") )
 		.Label( LOCTEXT("PreviewTabTitle", "Preview") )
 		.TabColorScale( GetTabColorScale() )
 		[
@@ -629,7 +629,7 @@ TSharedRef<SDockTab> FTranslationEditor::SpawnTab_Context( const FSpawnTabArgs&
 	TSharedRef<SWidget> PropertyTableWidget = ContextPropertyTableWidgetHandle->GetWidget();
 
 	TSharedRef<SDockTab> NewDockTab = SNew(SDockTab)
-		.Icon( FEditorStyle::GetBrush("TranslationEditor.Tabs.Properties") )
+		//.Icon( FEditorStyle::GetBrush("TranslationEditor.Tabs.Properties") )
 		.Label( LOCTEXT("ContextTabTitle", "Context") )
 		.TabColorScale( GetTabColorScale() )
 		[
@@ -725,7 +725,7 @@ TSharedRef<SDockTab> FTranslationEditor::SpawnTab_History(const FSpawnTabArgs& A
 	TSharedRef<SWidget> PropertyTableWidget = HistoryPropertyTableWidgetHandle->GetWidget();
 
 	TSharedRef<SDockTab> NewDockTab = SNew(SDockTab)
-		.Icon(FEditorStyle::GetBrush("TranslationEditor.Tabs.Properties"))
+		//.Icon(FEditorStyle::GetBrush("TranslationEditor.Tabs.Properties"))
 		.Label(LOCTEXT("HistoryTabTitle", "History"))
 		.TabColorScale(GetTabColorScale())
 		[
