FText: Method to retrieve source FTexts from FText::Format history.  Used in the translation picker currently.

TranslationPicker: Improving layout to be more useful for multi-line text, and to show text with no localization information in a more concise fashion.

diff --git a/Engine/Source/Editor/TranslationEditor/Private/TranslationPickerEditWindow.cpp b/Engine/Source/Editor/TranslationEditor/Private/TranslationPickerEditWindow.cpp
--- a/Engine/Source/Editor/TranslationEditor/Private/TranslationPickerEditWindow.cpp
+++ b/Engine/Source/Editor/TranslationEditor/Private/TranslationPickerEditWindow.cpp
@@ -40,7 +40,7 @@ void STranslationPickerEditWindow::Construct(const FArguments& InArgs)
 		EditWidgets.Add(NewEditWidget);
 
 		TextsBox->AddSlot()
-			.FillHeight(1)
+			.AutoHeight()
 			.Padding(FMargin(5))
 			[
 				SNew(SBorder)
@@ -60,9 +60,11 @@ void STranslationPickerEditWindow::Construct(const FArguments& InArgs)
 		[
 			SNew(SVerticalBox)
 			+ SVerticalBox::Slot()
-			.FillHeight(1)
+			.AutoHeight()
 			[
-				SNew(SBox)
+				SNew(SVerticalBox)
+				+ SVerticalBox::Slot()
+				.AutoHeight()
 				.Padding(FMargin(8, 5, 8, 5))
 				[
 					SNew(SVerticalBox)
@@ -173,6 +175,9 @@ void STranslationPickerEditWidget::Construct(const FArguments& InArgs)
 	bAllowEditing = InArgs._bAllowEditing;
 	int32 DefaultPadding = 0.0f;
 
+	bool bCultureInvariant = PickedText.IsCultureInvariant();
+	bool bShouldGatherForLocalization = FTextInspector::ShouldGatherForLocalization(PickedText);
+
 	// Get all the data we need and format it properly
 	const FString* NamespaceString = FTextInspector::GetNamespace(PickedText);
 	const FString* KeyString = FTextInspector::GetKey(PickedText);
@@ -211,93 +216,143 @@ void STranslationPickerEditWidget::Construct(const FArguments& InArgs)
 	TranslationUnit->LocresPath = LocresFullPath;
 	
 	// Can only save if we have all the required information
-	bool bCanSave = NamespaceString != nullptr && SourceString != nullptr && LocresFullPath.Len() > 0;
+	bool bHasRequiredLocalizationInfo = NamespaceString != nullptr && SourceString != nullptr && LocresFullPath.Len() > 0;
 
-	// Layout all our data
-	ChildSlot
-		[
-			SNew(SVerticalBox)
-			+SVerticalBox::Slot()
-			.FillHeight(1)
+	TSharedRef<SHorizontalBox> LocalizationInfoAndSaveButtonSlot = SNew(SHorizontalBox);
+
+	if (bCultureInvariant)
+	{
+		LocalizationInfoAndSaveButtonSlot->AddSlot()
+			.AutoWidth()
+			.Padding(FMargin(5))
+			[
+				SNew(STextBlock)
+				.Text(LOCTEXT("CultureInvariantLabel","This text is culture-invariant"))
+				.Justification(ETextJustify::Center)
+			];
+	}
+	else if (!bShouldGatherForLocalization)
+	{
+		LocalizationInfoAndSaveButtonSlot->AddSlot()
+			.AutoWidth()
+			.Padding(FMargin(5))
+			[
+				SNew(STextBlock)
+				.Text(LOCTEXT("NotGatheredForLocalizationLabel", "This text is not gathered for localization"))
+				.Justification(ETextJustify::Center)
+			];
+	}
+	else if (!bHasRequiredLocalizationInfo)
+	{
+		LocalizationInfoAndSaveButtonSlot->AddSlot()
+			.AutoWidth()
+			.Padding(FMargin(5))
+			[
+				SNew(STextBlock)
+				.Text(LOCTEXT("RequiredLocalizationInfoNotFound", "The required localization info for this text was not found."))
+				.Justification(ETextJustify::Center)
+			];
+	}
+	else
+	{
+		LocalizationInfoAndSaveButtonSlot->AddSlot()
+			.AutoWidth()
 			.Padding(FMargin(5))
 			[
 				SNew(SVerticalBox)
 				+ SVerticalBox::Slot()
-				.FillHeight(1)
-				.Padding(FMargin(5))
 				[
 					SNew(STextBlock)
-					.Text(SourceWithLabel)
+					.Text(NamespaceWithLabel)
 				]
 				+ SVerticalBox::Slot()
-				.FillHeight(1)
-				.Padding(FMargin(5))
+					[
+						SNew(STextBlock)
+						.Text(KeyWithLabel)
+					]
+				+ SVerticalBox::Slot()
+					[
+						SNew(STextBlock)
+						.Text(ManifestAndArchiveNameWithLabel)
+					]
+			];
+		LocalizationInfoAndSaveButtonSlot->AddSlot()
+			.AutoWidth()
+			.Padding(FMargin(5))
+			[
+				SNew(SButton)
+				.HAlign(HAlign_Center)
+				.ContentPadding(FEditorStyle::GetMargin("StandardDialog.ContentPadding"))
+				.OnClicked(this, &STranslationPickerEditWidget::SaveAndPreview)
+				.IsEnabled(bHasRequiredLocalizationInfo)
+				.Visibility(bAllowEditing ? EVisibility::Visible : EVisibility::Collapsed)
 				[
 					SNew(SHorizontalBox)
 					+ SHorizontalBox::Slot()
+					.Padding(FMargin(0, 0, 3, 0))
+					.VAlign(VAlign_Center)
 					.AutoWidth()
 					[
 						SNew(STextBlock)
-						.Text(LOCTEXT("TranslationLabel", "Translation: "))
-					]
-					+ SHorizontalBox::Slot()
-					[
-						SAssignNew(TextBox, SEditableTextBox)
-						.IsEnabled(bAllowEditing && bCanSave)
-						.Text(Translation)
-						.HintText(LOCTEXT("TranslationEditTextBox_HintText", "Enter/edit translation here."))
+						.Text(bHasRequiredLocalizationInfo ? LOCTEXT("SaveAndPreviewButtonText", "Save and preview") : LOCTEXT("SaveAndPreviewButtonDisabledText", "Cannot Save"))
 					]
 				]
-			]
+			];
+	}
+
+	TSharedRef<SHorizontalBox> TranslationSlot = SNew(SHorizontalBox);
+	
+	if (!bHasRequiredLocalizationInfo && SourceString->Equals(TranslationString))
+	{
+		// If loc info lookup failed and the translated and source string are the same, don't bother displaying the translation
+		TranslationSlot->SetVisibility(EVisibility::Collapsed);
+	}
+	else
+	{
+		TranslationSlot->AddSlot()
+			.AutoWidth()
+			[
+				SNew(STextBlock)
+				.Text(LOCTEXT("TranslationLabel", "Translation: "))
+			];
+		TranslationSlot->AddSlot()
+			[
+				SAssignNew(TextBox, SMultiLineEditableTextBox)
+				.IsEnabled(bAllowEditing && bHasRequiredLocalizationInfo)
+				.Text(Translation)
+				.HintText(LOCTEXT("TranslationEditTextBox_HintText", "Enter/edit translation here."))
+			];
+	}
+
+	// Layout all our data
+	ChildSlot
+		[
+			SNew(SVerticalBox)
 			+SVerticalBox::Slot()
-			.FillHeight(1)
+			.AutoHeight()
 			.Padding(FMargin(5))
 			[
-				SNew(SHorizontalBox)
-				+ SHorizontalBox::Slot()
-				.AutoWidth()
+				SNew(SVerticalBox)
+				+ SVerticalBox::Slot()
+				.AutoHeight()
 				.Padding(FMargin(5))
 				[
-					SNew(SVerticalBox)
-					+ SVerticalBox::Slot()
-					[
-						SNew(STextBlock)
-						.Text(NamespaceWithLabel)
-					]
-					+ SVerticalBox::Slot()
-					[
-						SNew(STextBlock)
-						.Text(KeyWithLabel)
-					]
-					+ SVerticalBox::Slot()
-					[
-						SNew(STextBlock)
-						.Text(ManifestAndArchiveNameWithLabel)
-					]
+					SNew(STextBlock)
+					.Text(SourceWithLabel)
 				]
-				+ SHorizontalBox::Slot()
-				.AutoWidth()
+				+ SVerticalBox::Slot()
+				.AutoHeight()
 				.Padding(FMargin(5))
 				[
-					SNew(SButton)
-					.HAlign(HAlign_Center)
-					.ContentPadding(FEditorStyle::GetMargin("StandardDialog.ContentPadding"))
-					.OnClicked(this, &STranslationPickerEditWidget::SaveAndPreview)
-					.IsEnabled(bCanSave)
-					.Visibility(bAllowEditing ? EVisibility::Visible : EVisibility::Collapsed)
-					[
-						SNew(SHorizontalBox)
-						+ SHorizontalBox::Slot()
-						.Padding(FMargin(0, 0, 3, 0))
-						.VAlign(VAlign_Center)
-						.AutoWidth()
-						[
-							SNew(STextBlock)
-							.Text(bCanSave ? LOCTEXT("SaveAndPreviewButtonText", "Save and preview") : LOCTEXT("SaveAndPreviewButtonDisabledText", "Cannot Save") )
-						]
-					]
+					TranslationSlot
 				]
 			]
+			+SVerticalBox::Slot()
+			.AutoHeight()
+			.Padding(FMargin(5))
+			[
+				LocalizationInfoAndSaveButtonSlot
+			]
 		];
 }
 
diff --git a/Engine/Source/Editor/TranslationEditor/Private/TranslationPickerEditWindow.h b/Engine/Source/Editor/TranslationEditor/Private/TranslationPickerEditWindow.h
--- a/Engine/Source/Editor/TranslationEditor/Private/TranslationPickerEditWindow.h
+++ b/Engine/Source/Editor/TranslationEditor/Private/TranslationPickerEditWindow.h
@@ -37,7 +37,7 @@ private:
 	UTranslationUnit* TranslationUnit;
 
 	/** The text box for entering/modifying a translation */
-	TSharedPtr<SEditableTextBox> TextBox;
+	TSharedPtr<SMultiLineEditableTextBox> TextBox;
 
 	/** Whether or not to show the save button*/
 	bool bAllowEditing;
diff --git a/Engine/Source/Editor/TranslationEditor/Private/TranslationPickerFloatingWindow.cpp b/Engine/Source/Editor/TranslationEditor/Private/TranslationPickerFloatingWindow.cpp
--- a/Engine/Source/Editor/TranslationEditor/Private/TranslationPickerFloatingWindow.cpp
+++ b/Engine/Source/Editor/TranslationEditor/Private/TranslationPickerFloatingWindow.cpp
@@ -65,7 +65,7 @@ void STranslationPickerFloatingWindow::Tick( const FGeometry& AllottedGeometry,
 			for (FText PickedText : PickedTexts)
 			{
 				TextsBox->AddSlot()
-					.FillHeight(1)
+					.AutoHeight()
 					.Padding(FMargin(5))
 					[
 						SNew(SBorder)
@@ -153,20 +153,31 @@ FText STranslationPickerFloatingWindow::GetTextFromWidget(TSharedRef<SWidget> Wi
 		}
 	}
 	
-	// Don't show the same text twice
-	bool bAlreadyPicked = false;
-	for (FText PickedText : PickedTexts)
+	if (!OriginalText.IsEmpty())
 	{
-		if (OriginalText.EqualTo(PickedText))
+		// Search the text from this widget's FText::Format history to find any source text
+		TArray<FText> FormattingSourceTexts;
+		OriginalText.GetSourceTextsFromFormatHistory(FormattingSourceTexts);
+
+		for (FText FormattingSourceText : FormattingSourceTexts)
 		{
-			bAlreadyPicked = true;
+			// Don't show the same text twice
+			bool bAlreadyPicked = false;
+			for (FText PickedText : PickedTexts)
+			{
+				if (FormattingSourceText.EqualTo(PickedText))
+				{
+					bAlreadyPicked = true;
+					break;
+				}
+			}
+
+			if (!bAlreadyPicked)
+			{
+				PickedTexts.Add(FormattingSourceText);
+			}
 		}
 	}
-	
-	if (!OriginalText.IsEmpty() && !bAlreadyPicked)
-	{
-		PickedTexts.Add(OriginalText);
-	}
 
 	return OriginalText;
 }
diff --git a/Engine/Source/Runtime/Core/Private/Internationalization/Text.cpp b/Engine/Source/Runtime/Core/Private/Internationalization/Text.cpp
--- a/Engine/Source/Runtime/Core/Private/Internationalization/Text.cpp
+++ b/Engine/Source/Runtime/Core/Private/Internationalization/Text.cpp
@@ -809,6 +809,11 @@ bool FText::IdenticalTo( const FText& Other ) const
 	return DisplayString == Other.DisplayString;
 }
 
+void FText::GetSourceTextsFromFormatHistory(TArray<FText>& OutSourceTexts) const
+{
+	History->GetSourceTextsFromFormatHistory(*(this), OutSourceTexts);
+}
+
 FTextSnapshot::FTextSnapshot()
 	: DisplayStringPtr()
 	, HistoryRevision(INDEX_NONE)
diff --git a/Engine/Source/Runtime/Core/Private/Internationalization/TextHistory.cpp b/Engine/Source/Runtime/Core/Private/Internationalization/TextHistory.cpp
--- a/Engine/Source/Runtime/Core/Private/Internationalization/TextHistory.cpp
+++ b/Engine/Source/Runtime/Core/Private/Internationalization/TextHistory.cpp
@@ -25,6 +25,12 @@ TSharedPtr< FString, ESPMode::ThreadSafe > FTextHistory::GetSourceString() const
 	return NULL;
 }
 
+void FTextHistory::GetSourceTextsFromFormatHistory(FText Text, TArray<FText>& OutSourceTexts) const
+{
+	// If we get here, we have no more source so this must be the base text
+	OutSourceTexts.Add(Text);
+}
+
 void FTextHistory::SerializeForDisplayString(FArchive& Ar, TSharedRef<FString, ESPMode::ThreadSafe>& InOutDisplayString)
 {
 	if(Ar.IsLoading())
@@ -201,6 +207,25 @@ void FTextHistory_NamedFormat::Serialize( FArchive& Ar )
 	Ar << Arguments;
 }
 
+void FTextHistory_NamedFormat::GetSourceTextsFromFormatHistory(FText, TArray<FText>& OutSourceTexts) const
+{
+	// Search the formatting text itself for source text
+	SourceText.GetSourceTextsFromFormatHistory(OutSourceTexts);
+
+	for (auto It = Arguments.CreateConstIterator(); It; ++It)
+	{
+		const FFormatArgumentValue& ArgumentValue = It.Value();
+		if (ArgumentValue.Type == EFormatArgumentType::Text)
+		{
+			if (ArgumentValue.TextValue && ArgumentValue.TextValue->History.IsValid())
+			{
+				// Search any text arguments for source text
+				ArgumentValue.TextValue->History->GetSourceTextsFromFormatHistory(*(ArgumentValue.TextValue), OutSourceTexts);
+			}
+		}
+	}
+}
+
 ///////////////////////////////////////
 // FTextHistory_OrderedFormat
 
@@ -228,6 +253,25 @@ void FTextHistory_OrderedFormat::Serialize( FArchive& Ar )
 	Ar << Arguments;
 }
 
+void FTextHistory_OrderedFormat::GetSourceTextsFromFormatHistory(FText, TArray<FText>& OutSourceTexts) const
+{
+	// Search the formatting text itself for source text
+	SourceText.GetSourceTextsFromFormatHistory(OutSourceTexts);
+
+	for (auto It = Arguments.CreateConstIterator(); It; ++It)
+	{
+		const FFormatArgumentValue& ArgumentValue = *It;
+		if (ArgumentValue.Type == EFormatArgumentType::Text)
+		{
+			if (ArgumentValue.TextValue && ArgumentValue.TextValue->History.IsValid())
+			{
+				// Search any text arguments for source text
+				ArgumentValue.TextValue->History->GetSourceTextsFromFormatHistory(*(ArgumentValue.TextValue), OutSourceTexts);
+			}
+		}
+	}
+}
+
 ///////////////////////////////////////
 // FTextHistory_ArgumentDataFormat
 
@@ -255,6 +299,25 @@ void FTextHistory_ArgumentDataFormat::Serialize( FArchive& Ar )
 	Ar << Arguments;
 }
 
+void FTextHistory_ArgumentDataFormat::GetSourceTextsFromFormatHistory(FText, TArray<FText>& OutBaseTexts) const
+{
+	// Search the formatting text itself for source text
+	SourceText.GetSourceTextsFromFormatHistory(OutBaseTexts);
+
+	for (int32 x = 0; x < Arguments.Num(); ++x)
+	{
+		const FFormatArgumentValue& ArgumentValue = Arguments[x].ArgumentValue;
+		if (ArgumentValue.Type == EFormatArgumentType::Text)
+		{
+			if (ArgumentValue.TextValue && ArgumentValue.TextValue->History.IsValid())
+			{
+				// Search any text arguments for source text
+				ArgumentValue.TextValue->History->GetSourceTextsFromFormatHistory(*(ArgumentValue.TextValue), OutBaseTexts);
+			}
+		}
+	}
+}
+
 ///////////////////////////////////////
 // FTextHistory_FormatNumber
 
diff --git a/Engine/Source/Runtime/Core/Private/Internationalization/TextHistory.h b/Engine/Source/Runtime/Core/Private/Internationalization/TextHistory.h
--- a/Engine/Source/Runtime/Core/Private/Internationalization/TextHistory.h
+++ b/Engine/Source/Runtime/Core/Private/Internationalization/TextHistory.h
@@ -42,6 +42,9 @@ public:
 	/** Returns the source string managed by the history (if any). */
 	virtual TSharedPtr< FString, ESPMode::ThreadSafe > GetSourceString() const;
 
+	/** Trace the history of this Text until we find the base Texts it was comprised from */
+	virtual void GetSourceTextsFromFormatHistory(FText Text, TArray<FText>& OutSourceTexts) const;
+
 	/** Will rebuild the display string if out of date. */
 	void Rebuild(TSharedRef< FString, ESPMode::ThreadSafe > InDisplayString);
 
@@ -82,6 +85,7 @@ public:
 	// Begin FTextHistory interface
 	virtual FText ToText(bool bInAsSource) const override;
 	virtual void Serialize(FArchive& Ar) override;
+	virtual void GetSourceTextsFromFormatHistory(FText, TArray<FText>& OutSourceTexts) const override;
 	// End FTextHistory interface
 
 private:
@@ -101,6 +105,7 @@ public:
 	// Begin FTextHistory interface
 	virtual FText ToText(bool bInAsSource) const override;
 	virtual void Serialize(FArchive& Ar) override;
+	virtual void GetSourceTextsFromFormatHistory(FText, TArray<FText>& OutSourceTexts) const override;
 	// End FTextHistory interface
 
 private:
@@ -120,6 +125,7 @@ public:
 	// Begin FTextHistory interface
 	virtual FText ToText(bool bInAsSource) const override;
 	virtual void Serialize(FArchive& Ar) override;
+	virtual void GetSourceTextsFromFormatHistory(FText, TArray<FText>& OutSourceTexts) const override;
 	// End FTextHistory interface
 
 private:
diff --git a/Engine/Source/Runtime/Core/Public/Internationalization/Text.h b/Engine/Source/Runtime/Core/Public/Internationalization/Text.h
--- a/Engine/Source/Runtime/Core/Public/Internationalization/Text.h
+++ b/Engine/Source/Runtime/Core/Public/Internationalization/Text.h
@@ -238,6 +238,9 @@ public:
 	 */
 	bool IdenticalTo( const FText& Other ) const;
 
+	/** Trace the history of this Text until we find the base Texts it was comprised from */
+	void GetSourceTextsFromFormatHistory(TArray<FText>& OutSourceTexts) const;
+
 	class CORE_API FSortPredicate
 	{
 	public:
