Translation Picker:
- Improve picking of text in slate widgets that are hittest invisible by searching recursively all children of the deepest widget returned by the hittest.
- Added scrolling support to accomodate the large number of texts this may return as a result in some cases.
- Changes in the layout to make fields easier to read.

diff --git a/Engine/Source/Editor/TranslationEditor/Private/TranslationPickerEditWindow.cpp b/Engine/Source/Editor/TranslationEditor/Private/TranslationPickerEditWindow.cpp
--- a/Engine/Source/Editor/TranslationEditor/Private/TranslationPickerEditWindow.cpp
+++ b/Engine/Source/Editor/TranslationEditor/Private/TranslationPickerEditWindow.cpp
@@ -7,24 +7,16 @@
 
 #define LOCTEXT_NAMESPACE "TranslationPicker"
 
+// Default dimensions of the Translation Picker edit window (floating window also uses these sizes, so it matches roughly)
+const int32 STranslationPickerEditWindow::DefaultEditWindowWidth = 500;
+const int32 STranslationPickerEditWindow::DefaultEditWindowHeight = 500;
+
 void STranslationPickerEditWindow::Construct(const FArguments& InArgs)
 {
 	ParentWindow = InArgs._ParentWindow;
 	PickedTexts = InArgs._PickedTexts;
-
 	WindowContents = SNew(SBox);
-
-	auto TextsBox = SNew(SVerticalBox);
-
-	// Display name of the current language
-	TextsBox->AddSlot()
-		.AutoHeight()
-		.Padding(FMargin(5))
-		[
-			SNew(STextBlock)
-			.Text(FText::FromString(FInternationalization::Get().GetCurrentCulture()->GetDisplayName()))
-			.Justification(ETextJustify::Center)
-		];
+	TSharedRef<SVerticalBox> TextsBox = SNew(SVerticalBox);
 
 	// Add a new Translation Picker Edit Widget for each picked text
 	for (FText PickedText : PickedTexts)
@@ -58,61 +50,70 @@ void STranslationPickerEditWindow::Construct(const FArguments& InArgs)
 		SNew(SBorder)
 		.BorderImage(FEditorStyle::GetBrush("ToolPanel.GroupBorder"))
 		[
+
 			SNew(SVerticalBox)
 			+ SVerticalBox::Slot()
 			.AutoHeight()
+			.HAlign(HAlign_Center)
+			.Padding(DefaultPadding)
 			[
-				SNew(SVerticalBox)
-				+ SVerticalBox::Slot()
-				.AutoHeight()
-				.Padding(FMargin(8, 5, 8, 5))
+				// Display name of the current language
+				SNew(STextBlock)
+				.Text(FText::FromString(FInternationalization::Get().GetCurrentCulture()->GetDisplayName()))
+				.Justification(ETextJustify::Center)
+			]
+			+ SVerticalBox::Slot()
+			[
+				SNew(SScrollBox)
+				+ SScrollBox::Slot()
+				.Padding(0.0f)
 				[
 					SNew(SVerticalBox)
 					+ SVerticalBox::Slot()
-					.Padding(DefaultPadding)
 					.AutoHeight()
+					.Padding(FMargin(8, 5, 8, 5))
 					[
 						TextsBox
 					]
-					+ SVerticalBox::Slot()
-					.Padding(DefaultPadding)
-					.AutoHeight()
-					.HAlign(HAlign_Right)
+				]
+			]
+			+ SVerticalBox::Slot()
+			.AutoHeight()
+			.HAlign(HAlign_Right)
+			.Padding(DefaultPadding)
+			[
+				SNew(SUniformGridPanel)
+				.SlotPadding(FEditorStyle::GetMargin("StandardDialog.SlotPadding"))
+				.MinDesiredSlotWidth(FEditorStyle::GetFloat("StandardDialog.MinDesiredSlotWidth"))
+				.MinDesiredSlotHeight(FEditorStyle::GetFloat("StandardDialog.MinDesiredSlotHeight"))
+				+ SUniformGridPanel::Slot(0, 0)
+				[
+					SNew(SButton)
+					.HAlign(HAlign_Center)
+					.ContentPadding(FEditorStyle::GetMargin("StandardDialog.ContentPadding"))
+					.OnClicked(this, &STranslationPickerEditWindow::SaveAllAndClose)
 					[
-						SNew(SUniformGridPanel)
-						.SlotPadding(FEditorStyle::GetMargin("StandardDialog.SlotPadding"))
-						.MinDesiredSlotWidth(FEditorStyle::GetFloat("StandardDialog.MinDesiredSlotWidth"))
-						.MinDesiredSlotHeight(FEditorStyle::GetFloat("StandardDialog.MinDesiredSlotHeight"))
-						+ SUniformGridPanel::Slot(0, 0)
-						[
-							SNew(SButton)
-							.HAlign(HAlign_Center)
-							.ContentPadding(FEditorStyle::GetMargin("StandardDialog.ContentPadding"))
-							.OnClicked(this, &STranslationPickerEditWindow::SaveAllAndClose)
-							[
-								SNew(SHorizontalBox)
-								+ SHorizontalBox::Slot()
-								.Padding(FMargin(0, 0, 3, 0))
-								.VAlign(VAlign_Center)
-								.AutoWidth()
-								[
-									SNew(STextBlock).Text(LOCTEXT("SaveAllAndClose", "Save all and close"))
-								]
-							]
-						]
-						+ SUniformGridPanel::Slot(1, 0)
+						SNew(SHorizontalBox)
+						+ SHorizontalBox::Slot()
+						.Padding(FMargin(0, 0, 3, 0))
+						.VAlign(VAlign_Center)
+						.AutoWidth()
 						[
-							SNew(SButton)
-							.HAlign(HAlign_Center)
-							.ContentPadding(FEditorStyle::GetMargin("StandardDialog.ContentPadding"))
-							.OnClicked(this, &STranslationPickerEditWindow::Close)
-							.Text(LOCTEXT("CancelButton", "Cancel"))
+							SNew(STextBlock).Text(LOCTEXT("SaveAllAndClose", "Save all and close"))
 						]
 					]
 				]
+				+ SUniformGridPanel::Slot(1, 0)
+				[
+					SNew(SButton)
+					.HAlign(HAlign_Center)
+					.ContentPadding(FEditorStyle::GetMargin("StandardDialog.ContentPadding"))
+					.OnClicked(this, &STranslationPickerEditWindow::Close)
+					.VAlign(VAlign_Center)
+					.Text(LOCTEXT("CancelButton", "Cancel"))
+				]
 			]
 		]
-
 	);
 
 	ChildSlot
@@ -202,10 +203,10 @@ void STranslationPickerEditWidget::Construct(const FArguments& InArgs)
 	FText ManifestAndArchiveName = FText::FromString(ManifestAndArchiveNameString);
 	FText Translation = FText::FromString(TranslationString);
 
-	FText SourceWithLabel = FText::Format(LOCTEXT("SourceLabel", "Source: {0}"), Source);
-	FText NamespaceWithLabel = FText::Format(LOCTEXT("NamespaceLabel", "Namespace: {0}"), Namespace);
-	FText KeyWithLabel = FText::Format(LOCTEXT("KeyLabel", "Key: {0}"), Key);
-	FText ManifestAndArchiveNameWithLabel = FText::Format(LOCTEXT("LocresFileLabel", "Manifest/Archive File: {0}"), FText::FromString(ManifestAndArchiveNameString));
+	FText SourceLabel = LOCTEXT("SourceLabel", "Source:");
+	FText NamespaceLabel = LOCTEXT("NamespaceLabel", "Namespace:");
+	FText KeyLabel = LOCTEXT("KeyLabel", "Key:");
+	FText ManifestAndArchiveNameLabel = LOCTEXT("LocresFileLabel", "Target :");
 
 	// Save the necessary data in UTranslationUnit for later.  This is what we pass to TranslationDataManager to save our edits
 	TranslationUnit = NewObject<UTranslationUnit>();
@@ -213,28 +214,95 @@ void STranslationPickerEditWidget::Construct(const FArguments& InArgs)
 	TranslationUnit->Source = SourceString != nullptr ? *SourceString : "";
 	TranslationUnit->Translation = TranslationString;
 	TranslationUnit->LocresPath = LocresFullPath;
-	
+
 	// Can only save if we have all the required information
 	bool bHasRequiredLocalizationInfo = NamespaceString.IsSet() && SourceString != nullptr && LocresFullPath.Len() > 0;
 
-	TSharedRef<SHorizontalBox> LocalizationInfoAndSaveButtonSlot = SNew(SHorizontalBox);
+	TSharedPtr<SGridPanel> GridPanel;
+	TSharedRef<SGridPanel> LocalizationInfoAndSaveButtonSlot = SNew(SGridPanel).FillColumn(2,1);
+
+	// Layout all our data
+	ChildSlot
+		[
+			SNew(SHorizontalBox)
+			+ SHorizontalBox::Slot()
+			.FillWidth(1)
+			.Padding(FMargin(5))
+			[
+				SNew(SVerticalBox)
+				+ SVerticalBox::Slot()
+				[
+				SAssignNew(GridPanel, SGridPanel)
+				.FillColumn(1,1)
+				+ SGridPanel::Slot(0,0)
+				.Padding(FMargin(5))
+				.HAlign(HAlign_Right)
+				[
+						SNew(STextBlock)
+						.TextStyle(FEditorStyle::Get(), "RichTextBlock.Bold")
+						.Text(SourceLabel)
+				]
+				+ SGridPanel::Slot(0, 1)
+					.Padding(FMargin(5))
+					.HAlign(HAlign_Right)
+				[
+					SNew(SVerticalBox)
+					.Visibility(!bHasRequiredLocalizationInfo && SourceString->Equals(TranslationString) ? EVisibility::Collapsed : EVisibility::Visible)
+					+ SVerticalBox::Slot()
+					[
+						SNew(STextBlock)
+						.TextStyle(FEditorStyle::Get(), "RichTextBlock.Bold")
+						.Text(LOCTEXT("TranslationLabel", "Translation: "))
+					]
+				]
+				
+				+ SGridPanel::Slot(1, 0)
+					.Padding(FMargin(5))
+					.ColumnSpan(2)
+				[
+					SNew(STextBlock)
+					.Text(Source)
+				]
+				+ SGridPanel::Slot(1, 1)
+					.ColumnSpan(2)
+					.Padding(FMargin(5))
+					[
+						SNew(SVerticalBox)
+						.Visibility(!bHasRequiredLocalizationInfo && SourceString->Equals(TranslationString) ? EVisibility::Collapsed : EVisibility::Visible)
+						+ SVerticalBox::Slot()
+						[
+							SAssignNew(TextBox, SMultiLineEditableTextBox)
+							.IsEnabled(bAllowEditing && bHasRequiredLocalizationInfo)
+							.Text(Translation)
+							.HintText(LOCTEXT("TranslationEditTextBox_HintText", "Enter/edit translation here."))
+						]
+					]
+				]
+			+ SVerticalBox::Slot()
+				.AutoHeight()
+				.Padding(FMargin(5))
+				[
+					LocalizationInfoAndSaveButtonSlot
+				]
+			]
+		];
 
 	if (bCultureInvariant)
 	{
-		LocalizationInfoAndSaveButtonSlot->AddSlot()
-			.AutoWidth()
+		LocalizationInfoAndSaveButtonSlot->AddSlot(0, 0)
 			.Padding(FMargin(5))
+			.ColumnSpan(2)
 			[
 				SNew(STextBlock)
-				.Text(LOCTEXT("CultureInvariantLabel","This text is culture-invariant"))
+				.Text(LOCTEXT("CultureInvariantLabel", "This text is culture-invariant"))
 				.Justification(ETextJustify::Center)
 			];
 	}
 	else if (!bShouldGatherForLocalization)
 	{
-		LocalizationInfoAndSaveButtonSlot->AddSlot()
-			.AutoWidth()
+		LocalizationInfoAndSaveButtonSlot->AddSlot(0, 0)
 			.Padding(FMargin(5))
+			.ColumnSpan(2)
 			[
 				SNew(STextBlock)
 				.Text(LOCTEXT("NotGatheredForLocalizationLabel", "This text is not gathered for localization"))
@@ -243,9 +311,9 @@ void STranslationPickerEditWidget::Construct(const FArguments& InArgs)
 	}
 	else if (!bHasRequiredLocalizationInfo)
 	{
-		LocalizationInfoAndSaveButtonSlot->AddSlot()
-			.AutoWidth()
+		LocalizationInfoAndSaveButtonSlot->AddSlot(0, 0)
 			.Padding(FMargin(5))
+			.ColumnSpan(2)
 			[
 				SNew(STextBlock)
 				.Text(LOCTEXT("RequiredLocalizationInfoNotFound", "The required localization info for this text was not found."))
@@ -254,30 +322,50 @@ void STranslationPickerEditWidget::Construct(const FArguments& InArgs)
 	}
 	else
 	{
-		LocalizationInfoAndSaveButtonSlot->AddSlot()
-			.AutoWidth()
-			.Padding(FMargin(5))
+		LocalizationInfoAndSaveButtonSlot->AddSlot(0, 0)
+			.Padding(FMargin(2.5))
+			.HAlign(HAlign_Right)
 			[
-				SNew(SVerticalBox)
-				+ SVerticalBox::Slot()
-				[
-					SNew(STextBlock)
-					.Text(NamespaceWithLabel)
-				]
-				+ SVerticalBox::Slot()
-					[
-						SNew(STextBlock)
-						.Text(KeyWithLabel)
-					]
-				+ SVerticalBox::Slot()
-					[
-						SNew(STextBlock)
-						.Text(ManifestAndArchiveNameWithLabel)
-					]
+				SNew(STextBlock)
+				.TextStyle(FEditorStyle::Get(), "RichTextBlock.Bold")
+				.Text(NamespaceLabel)
 			];
-		LocalizationInfoAndSaveButtonSlot->AddSlot()
-			.AutoWidth()
-			.Padding(FMargin(5))
+		LocalizationInfoAndSaveButtonSlot->AddSlot(1, 0)
+			.Padding(FMargin(2.5))
+			[
+				SNew(STextBlock)
+				.Text(Namespace)
+			];
+		LocalizationInfoAndSaveButtonSlot->AddSlot(0, 1)
+			.Padding(FMargin(2.5))
+			.HAlign(HAlign_Right)
+			[
+				SNew(STextBlock)
+				.TextStyle(FEditorStyle::Get(), "RichTextBlock.Bold")
+				.Text(KeyLabel)
+			];
+		LocalizationInfoAndSaveButtonSlot->AddSlot(1, 1)
+			.Padding(FMargin(2.5))
+			[
+				SNew(STextBlock)
+				.Text(Key)
+			];
+		LocalizationInfoAndSaveButtonSlot->AddSlot(0, 2)
+			.Padding(FMargin(2.5))
+			.HAlign(HAlign_Right)
+			[
+				SNew(STextBlock)
+				.TextStyle(FEditorStyle::Get(), "RichTextBlock.Bold")
+				.Text(ManifestAndArchiveNameLabel)
+			];
+		LocalizationInfoAndSaveButtonSlot->AddSlot(1, 2)
+			.Padding(FMargin(2.5))
+			[
+				SNew(STextBlock)
+				.Text(ManifestAndArchiveName)
+			];
+		LocalizationInfoAndSaveButtonSlot->AddSlot(2, 2)
+			.Padding(FMargin(2.5))
 			[
 				SNew(SButton)
 				.HAlign(HAlign_Center)
@@ -298,61 +386,6 @@ void STranslationPickerEditWidget::Construct(const FArguments& InArgs)
 				]
 			];
 	}
-
-	TSharedRef<SHorizontalBox> TranslationSlot = SNew(SHorizontalBox);
-	
-	if (!bHasRequiredLocalizationInfo && SourceString->Equals(TranslationString))
-	{
-		// If loc info lookup failed and the translated and source string are the same, don't bother displaying the translation
-		TranslationSlot->SetVisibility(EVisibility::Collapsed);
-	}
-	else
-	{
-		TranslationSlot->AddSlot()
-			.AutoWidth()
-			[
-				SNew(STextBlock)
-				.Text(LOCTEXT("TranslationLabel", "Translation: "))
-			];
-		TranslationSlot->AddSlot()
-			[
-				SAssignNew(TextBox, SMultiLineEditableTextBox)
-				.IsEnabled(bAllowEditing && bHasRequiredLocalizationInfo)
-				.Text(Translation)
-				.HintText(LOCTEXT("TranslationEditTextBox_HintText", "Enter/edit translation here."))
-			];
-	}
-
-	// Layout all our data
-	ChildSlot
-		[
-			SNew(SVerticalBox)
-			+SVerticalBox::Slot()
-			.AutoHeight()
-			.Padding(FMargin(5))
-			[
-				SNew(SVerticalBox)
-				+ SVerticalBox::Slot()
-				.AutoHeight()
-				.Padding(FMargin(5))
-				[
-					SNew(STextBlock)
-					.Text(SourceWithLabel)
-				]
-				+ SVerticalBox::Slot()
-				.AutoHeight()
-				.Padding(FMargin(5))
-				[
-					TranslationSlot
-				]
-			]
-			+SVerticalBox::Slot()
-			.AutoHeight()
-			.Padding(FMargin(5))
-			[
-				LocalizationInfoAndSaveButtonSlot
-			]
-		];
 }
 
 FReply STranslationPickerEditWidget::SaveAndPreview()
diff --git a/Engine/Source/Editor/TranslationEditor/Private/TranslationPickerEditWindow.h b/Engine/Source/Editor/TranslationEditor/Private/TranslationPickerEditWindow.h
--- a/Engine/Source/Editor/TranslationEditor/Private/TranslationPickerEditWindow.h
+++ b/Engine/Source/Editor/TranslationEditor/Private/TranslationPickerEditWindow.h
@@ -56,6 +56,10 @@ class STranslationPickerEditWindow : public SCompoundWidget
 
 	void Construct(const FArguments& InArgs);
 
+	// Default dimensions of the Translation Picker edit window (floating window also uses these sizes, so it matches roughly)
+	static const int32 DefaultEditWindowWidth;
+	static const int32 DefaultEditWindowHeight;
+
 private:
 
 	virtual FReply OnKeyDown(const FGeometry& MyGeometry, const FKeyEvent& InKeyEvent) override;
diff --git a/Engine/Source/Editor/TranslationEditor/Private/TranslationPickerFloatingWindow.cpp b/Engine/Source/Editor/TranslationEditor/Private/TranslationPickerFloatingWindow.cpp
--- a/Engine/Source/Editor/TranslationEditor/Private/TranslationPickerFloatingWindow.cpp
+++ b/Engine/Source/Editor/TranslationEditor/Private/TranslationPickerFloatingWindow.cpp
@@ -19,6 +19,25 @@ void STranslationPickerFloatingWindow::Construct(const FArguments& InArgs)
 	];
 }
 
+void STranslationPickerFloatingWindow::GetTextFromChildWidgets(TSharedRef<SWidget> Widget)
+{
+	FChildren* Children = Widget->GetChildren();
+
+	TArray<TSharedRef<SWidget>> ChildrenArray;
+	for (int ChildIndex = 0; ChildIndex < Children->Num(); ++ChildIndex)
+	{
+		ChildrenArray.Add(Children->GetChildAt(ChildIndex));
+	}
+
+	for (int ChildIndex = 0; ChildIndex < Children->Num(); ++ChildIndex)
+	{
+		// Pull out any FText from this child widget
+		GetTextFromWidget(Children->GetChildAt(ChildIndex));
+		// Recursively search this child's children for their FText
+		GetTextFromChildWidgets(Children->GetChildAt(ChildIndex));
+	}
+}
+
 void STranslationPickerFloatingWindow::Tick( const FGeometry& AllottedGeometry, const double InCurrentTime, const float InDeltaTime )
 {
 	FWidgetPath Path = FSlateApplication::Get().LocateWindowUnderMouse(FSlateApplication::Get().GetCursorPos(), FSlateApplication::Get().GetInteractiveTopLevelWindows(), true);
@@ -30,7 +49,6 @@ void STranslationPickerFloatingWindow::Tick( const FGeometry& AllottedGeometry,
 		{
 			// Clear all previous text and widgets
 			PickedTexts.Empty();
-			WindowContents->SetContentWidget(SNew(SBox));
 
 			// Search everything under the cursor for any FText we know how to parse
 			for (int32 PathIndex = Path.Widgets.Num() - 1; PathIndex >= 0; PathIndex--)
@@ -47,19 +65,17 @@ void STranslationPickerFloatingWindow::Tick( const FGeometry& AllottedGeometry,
 				{
 					GetTextFromWidget(Tooltip->AsWidget());
 				}
-			}
 
-			auto TextsBox = SNew(SVerticalBox);
+				// Currently LocateWindowUnderMouse doesn't return hit-test invisible widgets
+				// So recursively search all child widgets of the deepest widget in our path in case there is hit-test invisible text in the deeper children
+				// GetTextFromWidget prevents adding the same FText twice, so shouldn't be any harm
+				if (PathIndex == (Path.Widgets.Num() - 1))
+				{
+					GetTextFromChildWidgets(PathWidget);
+				}
+			}
 
-			// Display name of the current language
-			TextsBox->AddSlot()
-				.AutoHeight()
-				.Padding(FMargin(5))
-				[
-					SNew(STextBlock)
-					.Text(FText::FromString(FInternationalization::Get().GetCurrentCulture()->GetDisplayName()))
-					.Justification(ETextJustify::Center)
-				];
+			TSharedRef<SVerticalBox> TextsBox = SNew(SVerticalBox);
 
 			// Add a new Translation Picker Edit Widget for each picked text
 			for (FText PickedText : PickedTexts)
@@ -77,26 +93,41 @@ void STranslationPickerFloatingWindow::Tick( const FGeometry& AllottedGeometry,
 					];
 			}
 
-			FText Instructions = FText::GetEmpty();
-			if (PickedTexts.Num() >= 1)
-			{
-				Instructions = LOCTEXT("TranslationPickerEscToEdit", "Press Esc to edit this translation");
-			}
-			else
-			{
-				Instructions = LOCTEXT("TranslationPickerHoverToViewEditEscToQuit", "Hover over text to view/edit translation, or press Esc to quit");
-			}
-
-			TextsBox->AddSlot()
+			WindowContents->SetContentWidget(
+				SNew(SVerticalBox)
+				+ SVerticalBox::Slot()
+				.Padding(0)
 				.AutoHeight()
 				.Padding(FMargin(5))
 				[
 					SNew(STextBlock)
-					.Text(Instructions)
+					.Text(FText::FromString(FInternationalization::Get().GetCurrentCulture()->GetDisplayName()))
 					.Justification(ETextJustify::Center)
-				];
-
-			WindowContents->SetContentWidget(TextsBox);
+				]
+			+ SVerticalBox::Slot()
+				.Padding(0)
+				.FillHeight(1)
+				.Padding(FMargin(5))
+				[
+					SNew(SScrollBox)
+					.Orientation(EOrientation::Orient_Vertical)
+					.ScrollBarAlwaysVisible(true)
+					+SScrollBox::Slot()
+						.Padding(FMargin(0))
+						[
+									TextsBox
+						]
+				]
+			+ SVerticalBox::Slot()
+				.Padding(0)
+				.AutoHeight()
+				.Padding(FMargin(5))
+				[
+					SNew(STextBlock)
+					.Text(PickedTexts.Num() >= 1 ? LOCTEXT("TranslationPickerEscToEdit", "Press Esc to edit translation(s)") : LOCTEXT("TranslationPickerHoverToViewEditEscToQuit", "Hover over text to view/edit translation, or press Esc to quit"))
+					.Justification(ETextJustify::Center)
+				]
+			);
 		}
 	}
 
@@ -104,7 +135,7 @@ void STranslationPickerFloatingWindow::Tick( const FGeometry& AllottedGeometry,
 	FSlateApplication::Get().SetKeyboardFocus(SharedThis(this), EKeyboardFocusCause::SetDirectly);
 	if (ParentWindow.IsValid())
 	{
-		FVector2D WindowSize = ParentWindow.Pin()->GetDesiredSize();
+		FVector2D WindowSize = ParentWindow.Pin()->GetSizeInScreen();
 		FVector2D DesiredPosition = FSlateApplication::Get().GetCursorPos();
 		DesiredPosition.X -= FSlateApplication::Get().GetCursorSize().X;
 		DesiredPosition.Y += FSlateApplication::Get().GetCursorSize().Y;
@@ -123,14 +154,13 @@ FText STranslationPickerFloatingWindow::GetTextFromWidget(TSharedRef<SWidget> Wi
 {
 	FText OriginalText = FText::GetEmpty();
 
-	// Parse STextBlocks
 	STextBlock* TextBlock = (STextBlock*)&Widget.Get();
+
+	// Have to parse the various widget types to find the FText
 	if ((Widget->GetTypeAsString() == "STextBlock" && TextBlock))
 	{
 		OriginalText = TextBlock->GetText();
 	}
-
-	// Parse SToolTips
 	else if (Widget->GetTypeAsString() == "SToolTip")
 	{
 		SToolTip* ToolTip = ((SToolTip*)&Widget.Get());
@@ -143,7 +173,6 @@ FText STranslationPickerFloatingWindow::GetTextFromWidget(TSharedRef<SWidget> Wi
 			}
 		}
 	}
-	// Parse SDocumentationToolTips
 	else if (Widget->GetTypeAsString() == "SDocumentationToolTip")
 	{
 		SDocumentationToolTip* DocumentationToolTip = (SDocumentationToolTip*)&Widget.Get();
@@ -152,6 +181,57 @@ FText STranslationPickerFloatingWindow::GetTextFromWidget(TSharedRef<SWidget> Wi
 			OriginalText = DocumentationToolTip->GetTextTooltip();
 		}
 	}
+	else if (Widget->GetTypeAsString() == "SEditableText")
+	{
+		SEditableText* EditableText = (SEditableText*)&Widget.Get();
+		if (EditableText != nullptr)
+		{
+			// Always return the hint text because that's the only thing that will be translatable
+			OriginalText = EditableText->GetHintText();
+		}
+	}
+	else if (Widget->GetTypeAsString() == "SRichTextBlock")
+	{
+		SRichTextBlock* RichTextBlock = (SRichTextBlock*)&Widget.Get();
+		if (RichTextBlock != nullptr)
+		{
+			OriginalText = RichTextBlock->GetText();
+		}
+	}
+	else if (Widget->GetTypeAsString() == "SMultiLineEditableText")
+	{
+		SMultiLineEditableText* MultiLineEditableText = (SMultiLineEditableText*)&Widget.Get();
+		if (MultiLineEditableText != nullptr)
+		{
+			// Always return the hint text because that's the only thing that will be translatable
+			OriginalText = MultiLineEditableText->GetHintText();
+		}
+	}
+	else if (Widget->GetTypeAsString() == "SMultiLineEditableTextBox")
+	{
+		SMultiLineEditableTextBox* MultiLineEditableTextBox = (SMultiLineEditableTextBox*)&Widget.Get();
+		if (MultiLineEditableTextBox != nullptr)
+		{
+			OriginalText = MultiLineEditableTextBox->GetText();
+		}
+	}
+	else if (Widget->GetTypeAsString() == "SButton")
+	{
+		SButton* Button = (SButton*)&Widget.Get();
+
+		// It seems like the LocateWindowUnderMouse() function will sometimes return an SButton but not the FText inside the button?
+		// So try to find the first FText child of a button just in case
+		FChildren* Children = Button->GetChildren();
+		for (int ChildIndex = 0; ChildIndex < Children->Num(); ++ChildIndex)
+		{
+			TSharedRef<SWidget> ChildWidget = Children->GetChildAt(ChildIndex);
+			OriginalText = GetTextFromWidget(ChildWidget);
+			if (!OriginalText.IsEmpty())
+			{
+				break;
+			}
+		}
+	}
 	
 	if (!OriginalText.IsEmpty())
 	{
@@ -191,16 +271,17 @@ FReply STranslationPickerFloatingWindow::OnKeyDown(const FGeometry& MyGeometry,
 		{
 			TSharedRef<SWindow> NewWindow = SNew(SWindow)
 				.Title(LOCTEXT("TranslationPickerEditWindowTitle", "Edit Translation(s)"))
-				.SupportsMaximize(false)
-				.SupportsMinimize(false)
 				.CreateTitleBar(true)
-				.SizingRule(ESizingRule::Autosized);
+				.SizingRule(ESizingRule::UserSized);
 
-			NewWindow->SetContent(
-				SNew(STranslationPickerEditWindow)
+			TSharedRef<STranslationPickerEditWindow> EditWindow = SNew(STranslationPickerEditWindow)
 				.ParentWindow(NewWindow)
-				.PickedTexts(PickedTexts)
-				);
+				.PickedTexts(PickedTexts);
+
+			NewWindow->SetContent(EditWindow);
+
+			// Make this roughly the same size as the Edit Window, so when you press Esc to edit, the window is in basically the same size
+			NewWindow->Resize(FVector2D(STranslationPickerEditWindow::DefaultEditWindowWidth, STranslationPickerEditWindow::DefaultEditWindowHeight));
 
 			TSharedPtr<SWindow> RootWindow = FGlobalTabmanager::Get()->GetRootWindow();
 			if (RootWindow.IsValid())
@@ -212,7 +293,7 @@ FReply STranslationPickerFloatingWindow::OnKeyDown(const FGeometry& MyGeometry,
 				FSlateApplication::Get().AddWindow(NewWindow);
 			}
 
-			FVector2D WindowSize = ParentWindow.Pin()->GetDesiredSize();
+			FVector2D WindowSize = ParentWindow.Pin()->GetSizeInScreen();
 			FVector2D DesiredPosition = FSlateApplication::Get().GetCursorPos();
 			DesiredPosition.X -= FSlateApplication::Get().GetCursorSize().X;
 			DesiredPosition.Y += FSlateApplication::Get().GetCursorSize().Y;
diff --git a/Engine/Source/Editor/TranslationEditor/Private/TranslationPickerFloatingWindow.h b/Engine/Source/Editor/TranslationEditor/Private/TranslationPickerFloatingWindow.h
--- a/Engine/Source/Editor/TranslationEditor/Private/TranslationPickerFloatingWindow.h
+++ b/Engine/Source/Editor/TranslationEditor/Private/TranslationPickerFloatingWindow.h
@@ -22,6 +22,9 @@ private:
 	/** Pull the FText reference out of an SWidget */
 	FText GetTextFromWidget(TSharedRef<SWidget> Widget);
 
+	/** Pull the FText reference out of the child widgets of an SWidget */
+	void GetTextFromChildWidgets(TSharedRef<SWidget> Widget);
+
 	/** Handle key presses */
 	virtual FReply OnKeyDown(const FGeometry& MyGeometry, const FKeyEvent& InKeyEvent) override;
 
diff --git a/Engine/Source/Editor/TranslationEditor/Private/TranslationPickerWidget.cpp b/Engine/Source/Editor/TranslationEditor/Private/TranslationPickerWidget.cpp
--- a/Engine/Source/Editor/TranslationEditor/Private/TranslationPickerWidget.cpp
+++ b/Engine/Source/Editor/TranslationEditor/Private/TranslationPickerWidget.cpp
@@ -2,6 +2,7 @@
 
 #include "TranslationEditorPrivatePCH.h"
 #include "TranslationPickerWidget.h"
+#include "TranslationPickerEditWindow.h"
 
 #define LOCTEXT_NAMESPACE "TranslationPicker"
 
@@ -107,6 +108,9 @@ bool TranslationPickerManager::OpenPickerWindow()
 	if (!PickerWindow.IsValid())
 	{
 		TSharedRef<SWindow> NewWindow = SWindow::MakeCursorDecorator();
+		NewWindow->SetSizingRule(ESizingRule::FixedSize);
+		// The Edit window and Floating window should be roughly the same size, so it isn't too distracting switching between them
+		NewWindow->Resize(FVector2D(STranslationPickerEditWindow::DefaultEditWindowWidth, STranslationPickerEditWindow::DefaultEditWindowHeight));
 		NewWindow->MoveWindowTo(FSlateApplication::Get().GetCursorPos());
 		PickerWindow = NewWindow;
 
