Add "Import latest from Localization Service" button to translation editor

diff --git a/Engine/Plugins/Developer/OneSkyLocalizationService/Source/OneSkyLocalizationService/Private/OneSkyLocalizationServiceOperations.cpp b/Engine/Plugins/Developer/OneSkyLocalizationService/Source/OneSkyLocalizationService/Private/OneSkyLocalizationServiceOperations.cpp
--- a/Engine/Plugins/Developer/OneSkyLocalizationService/Source/OneSkyLocalizationService/Private/OneSkyLocalizationServiceOperations.cpp
+++ b/Engine/Plugins/Developer/OneSkyLocalizationService/Source/OneSkyLocalizationService/Private/OneSkyLocalizationServiceOperations.cpp
@@ -833,6 +833,11 @@ void FOneSkyTranslationExportWorker::Query_HttpRequestComplete(FHttpRequestPtr H
 		{
 			UE_LOG(LogLocalizationService, Warning, TEXT("%s"), *(ErrorText.ToString()));
 			Command->ErrorMessages.Add(ErrorText);
+			TSharedPtr<FDownloadLocalizationTargetFile, ESPMode::ThreadSafe> DownloadLocTargetOp = StaticCastSharedRef<FDownloadLocalizationTargetFile>(Command->Operation);
+			if (DownloadLocTargetOp.IsValid())
+			{
+				DownloadLocTargetOp->SetOutErrorText(ErrorText);
+			}
 		}
 	}
 
diff --git a/Engine/Plugins/Developer/OneSkyLocalizationService/Source/OneSkyLocalizationService/Private/OneSkyLocalizationServiceSettings.cpp b/Engine/Plugins/Developer/OneSkyLocalizationService/Source/OneSkyLocalizationService/Private/OneSkyLocalizationServiceSettings.cpp
--- a/Engine/Plugins/Developer/OneSkyLocalizationService/Source/OneSkyLocalizationService/Private/OneSkyLocalizationServiceSettings.cpp
+++ b/Engine/Plugins/Developer/OneSkyLocalizationService/Source/OneSkyLocalizationService/Private/OneSkyLocalizationServiceSettings.cpp
@@ -160,6 +160,8 @@ FOneSkyConnectionInfo FOneSkyLocalizationServiceSettings::GetConnectionInfo() co
 
 FOneSkyLocalizationTargetSetting* FOneSkyLocalizationServiceSettings::GetSettingsForTarget(FGuid TargetGuid, bool bCreateIfNotFound /* = false */)
 {
+	TargetSettingsObject->LoadConfig();
+
 	FOneSkyLocalizationTargetSetting* Settings = TargetSettingsObject->TargetSettings.FindByPredicate([&](const FOneSkyLocalizationTargetSetting& Setting)
 	{
 		return Setting.TargetGuid == TargetGuid;
diff --git a/Engine/Plugins/Developer/OneSkyLocalizationService/Source/OneSkyLocalizationService/Private/OneSkyLocalizationServiceSettings.h b/Engine/Plugins/Developer/OneSkyLocalizationService/Source/OneSkyLocalizationService/Private/OneSkyLocalizationServiceSettings.h
--- a/Engine/Plugins/Developer/OneSkyLocalizationService/Source/OneSkyLocalizationService/Private/OneSkyLocalizationServiceSettings.h
+++ b/Engine/Plugins/Developer/OneSkyLocalizationService/Source/OneSkyLocalizationService/Private/OneSkyLocalizationServiceSettings.h
@@ -35,7 +35,7 @@ struct FOneSkyLocalizationTargetSetting
 
 };
 
-UCLASS(config = Editor)
+UCLASS(config = LocalizationServiceSettings)
 class UOneSkyLocalizationTargetSettings : public UObject
 {
 	GENERATED_UCLASS_BODY()
diff --git a/Engine/Source/Developer/LocalizationService/Public/LocalizationServiceOperations.h b/Engine/Source/Developer/LocalizationService/Public/LocalizationServiceOperations.h
--- a/Engine/Source/Developer/LocalizationService/Public/LocalizationServiceOperations.h
+++ b/Engine/Source/Developer/LocalizationService/Public/LocalizationServiceOperations.h
@@ -100,6 +100,17 @@ public:
 		InLocale = NewLocale;
 	}
 
+	FText GetOutErrorText()
+	{
+		return OutErrorText;
+	}
+
+	void SetOutErrorText(const FText& NewOutErrorText)
+	{
+		OutErrorText = NewOutErrorText;
+	}
+
+
 protected:
 
 	/** The GUID of the Localization Target */
@@ -110,6 +121,9 @@ protected:
 
 	/** The path and name to the file downloaded, relative to project directory */
 	FString InRelativeOutputFilePathAndName;
+
+	/** Place to easily store and access error message */
+	FText OutErrorText;
 };
 
 ///**
diff --git a/Engine/Source/Editor/EditorStyle/Private/SlateEditorStyle.cpp b/Engine/Source/Editor/EditorStyle/Private/SlateEditorStyle.cpp
--- a/Engine/Source/Editor/EditorStyle/Private/SlateEditorStyle.cpp
+++ b/Engine/Source/Editor/EditorStyle/Private/SlateEditorStyle.cpp
@@ -6617,11 +6617,12 @@ void FSlateEditorStyle::FStyle::SetupUMGEditorStyles()
 
 void FSlateEditorStyle::FStyle::SetupTranslationEditorStyles()
 {
-	Set("TranslationEditor.Export", new IMAGE_BRUSH("Icons/icon_Import_40x", Icon40x40));
+	Set("TranslationEditor.Export", new IMAGE_BRUSH("Icons/Icon_Localisation_Export_All_40x", Icon40x40));
 	Set("TranslationEditor.PreviewInEditor", new IMAGE_BRUSH("Icons/icon_levels_visible_40x", Icon40x40));
-	Set("TranslationEditor.Import", new IMAGE_BRUSH("Icons/icon_Samples_Download_40x", Icon40x40));
+	Set("TranslationEditor.Import", new IMAGE_BRUSH("Icons/Icon_Localisation_Import_All_40x", Icon40x40));
 	Set("TranslationEditor.Search", new IMAGE_BRUSH("Icons/icon_Blueprint_Find_40px", Icon40x40));
 	Set("TranslationEditor.TranslationPicker", new IMAGE_BRUSH("Icons/icon_StaticMeshEd_VertColor_40x", Icon40x40));
+	Set("TranslationEditor.ImportLatestFromLocalizationService", new IMAGE_BRUSH("Icons/icon_worldscript_40x", Icon40x40));
 }
 
 
diff --git a/Engine/Source/Editor/LocalizationDashboard/Private/SLocalizationTargetEditorCultureRow.cpp b/Engine/Source/Editor/LocalizationDashboard/Private/SLocalizationTargetEditorCultureRow.cpp
--- a/Engine/Source/Editor/LocalizationDashboard/Private/SLocalizationTargetEditorCultureRow.cpp
+++ b/Engine/Source/Editor/LocalizationDashboard/Private/SLocalizationTargetEditorCultureRow.cpp
@@ -325,7 +325,7 @@ FReply SLocalizationTargetEditorCultureRow::Edit()
 	ULocalizationTarget* const LocalizationTarget = GetTarget();
 	if (Culture.IsValid() && LocalizationTarget)
 	{
-		FModuleManager::LoadModuleChecked<ITranslationEditor>("TranslationEditor").OpenTranslationEditor(LocalizationConfigurationScript::GetManifestPath(LocalizationTarget), LocalizationConfigurationScript::GetArchivePath(LocalizationTarget, Culture->GetName()));
+		FModuleManager::LoadModuleChecked<ITranslationEditor>("TranslationEditor").OpenTranslationEditor(LocalizationConfigurationScript::GetManifestPath(LocalizationTarget), LocalizationConfigurationScript::GetArchivePath(LocalizationTarget, Culture->GetName()), LocalizationTarget->Settings.Guid);
 	}
 
 	return FReply::Handled();
diff --git a/Engine/Source/Editor/TranslationEditor/Private/ITranslationEditor.cpp b/Engine/Source/Editor/TranslationEditor/Private/ITranslationEditor.cpp
--- a/Engine/Source/Editor/TranslationEditor/Private/ITranslationEditor.cpp
+++ b/Engine/Source/Editor/TranslationEditor/Private/ITranslationEditor.cpp
@@ -8,13 +8,13 @@
 /** To keep track of what translation editors are open editing which manifest files */
 TMap<FString, ITranslationEditor*> ITranslationEditor::OpenTranslationEditors = TMap<FString, ITranslationEditor*>();
 
-void ITranslationEditor::OpenTranslationEditor(const FString& InManifestFile, const FString& InArchiveFile)
+void ITranslationEditor::OpenTranslationEditor(const FString& InManifestFile, const FString& InArchiveFile, const FGuid& LocalizationTargetGuid)
 {
 	// Only open a new translation editor if there isn't another one open for this archive file already
 	if (!OpenTranslationEditors.Contains(InArchiveFile))
 	{
 		FTranslationEditorModule& TranslationEditorModule = FModuleManager::LoadModuleChecked<FTranslationEditorModule>("TranslationEditor");
-		ITranslationEditor* NewTranslationEditor = (ITranslationEditor*)&(TranslationEditorModule.CreateTranslationEditor(InManifestFile, InArchiveFile).Get());
+		ITranslationEditor* NewTranslationEditor = (ITranslationEditor*)&(TranslationEditorModule.CreateTranslationEditor(InManifestFile, InArchiveFile, LocalizationTargetGuid).Get());
 		NewTranslationEditor->RegisterTranslationEditor();
 	}
 	else
diff --git a/Engine/Source/Editor/TranslationEditor/Private/TranslationEditor.cpp b/Engine/Source/Editor/TranslationEditor/Private/TranslationEditor.cpp
--- a/Engine/Source/Editor/TranslationEditor/Private/TranslationEditor.cpp
+++ b/Engine/Source/Editor/TranslationEditor/Private/TranslationEditor.cpp
@@ -22,6 +22,8 @@
 #include "SSearchBox.h"
 #include "SDockTab.h"
 #include "InternationalizationExportSettings.h"
+#include "ILocalizationServiceModule.h"
+#include "ILocalizationServiceProvider.h"
 
 DEFINE_LOG_CATEGORY_STATIC(LocalizationExport, Log, All);
 
@@ -798,6 +800,10 @@ void FTranslationEditor::MapActions()
 		FExecuteAction::CreateSP(this, &FTranslationEditor::PreviewAllTranslationsInEditor_Execute),
 		FCanExecuteAction());
 
+	ToolkitCommands->MapAction(FTranslationEditorCommands::Get().ImportLatestFromLocalizationService,
+		FExecuteAction::CreateSP(this, &FTranslationEditor::ImportLatestFromLocalizationService_Execute),
+		FCanExecuteAction());
+
 	ToolkitCommands->MapAction(FTranslationEditorCommands::Get().ExportToPortableObjectFormat,
 		FExecuteAction::CreateSP(this, &FTranslationEditor::ExportToPortableObjectFormat_Execute),
 		FCanExecuteAction());
@@ -1129,6 +1135,99 @@ void FTranslationEditor::PreviewAllTranslationsInEditor_Execute()
 	DataManager->PreviewAllTranslationsInEditor();
 }
 
+void FTranslationEditor::ImportLatestFromLocalizationService_Execute()
+{
+	ILocalizationServiceProvider& Provider = ILocalizationServiceModule::Get().GetProvider();
+	TSharedRef<FDownloadLocalizationTargetFile, ESPMode::ThreadSafe> DownloadTargetFileOp = ILocalizationServiceOperation::Create<FDownloadLocalizationTargetFile>();
+	DownloadTargetFileOp->SetInTargetGuid(LocalizationTargetGuid);
+	DownloadTargetFileOp->SetInLocale(FPaths::GetBaseFilename(FPaths::GetPath(ArchiveFilePath)));
+	FString Path = FPaths::GameSavedDir() / "Temp" / "LastImportFromLocService.po";
+	FPaths::MakePathRelativeTo(Path, *FPaths::GameDir());
+	DownloadTargetFileOp->SetInRelativeOutputFilePathAndName(Path);
+
+	GWarn->BeginSlowTask(LOCTEXT("ImportingFromLocalizationService", "Importing Latest from Localization Service..."), true);
+
+	Provider.Execute(DownloadTargetFileOp, TArray<FLocalizationServiceTranslationIdentifier>(), ELocalizationServiceOperationConcurrency::Asynchronous, FLocalizationServiceOperationComplete::CreateSP(this, &FTranslationEditor::DownloadLatestFromLocalizationServiceComplete));
+}
+
+void FTranslationEditor::DownloadLatestFromLocalizationServiceComplete(const FLocalizationServiceOperationRef& Operation, ELocalizationServiceOperationCommandResult::Type Result)
+{
+	TSharedPtr<FDownloadLocalizationTargetFile, ESPMode::ThreadSafe> DownloadLocalizationTargetOp = StaticCastSharedRef<FDownloadLocalizationTargetFile>(Operation);
+	bool bError = !(Result == ELocalizationServiceOperationCommandResult::Succeeded);
+	FText ErrorText = FText::GetEmpty();
+	if (DownloadLocalizationTargetOp.IsValid())
+	{
+		ErrorText = DownloadLocalizationTargetOp->GetOutErrorText();
+	}
+	if (!bError && ErrorText.IsEmpty())
+	{
+		FGuid InTargetGuid;
+		FString InLocale;
+		FString InRelativeOutputFilePathAndName;
+
+		if (DownloadLocalizationTargetOp.IsValid())
+		{
+			InTargetGuid = DownloadLocalizationTargetOp->GetInTargetGuid();
+			InLocale = DownloadLocalizationTargetOp->GetInLocale();
+			InRelativeOutputFilePathAndName = DownloadLocalizationTargetOp->GetInRelativeOutputFilePathAndName();
+
+		}
+		else
+		{
+			bError = true;
+		}
+
+		if (InTargetGuid == LocalizationTargetGuid && InLocale == FPaths::GetBaseFilename(FPaths::GetPath(ArchiveFilePath)) && !InRelativeOutputFilePathAndName.IsEmpty())
+		{
+			FString AbsoluteFilePathAndName = FPaths::ConvertRelativePathToFull(FPaths::GameDir() / InRelativeOutputFilePathAndName);
+			if (FPaths::FileExists(AbsoluteFilePathAndName))
+			{
+				GWarn->StatusUpdate(50, 100, LOCTEXT("DownloadFromLocalizationServiceFinishedNowImporting", "Download from Localization Service Finished, Importing..."));
+				ImportFromPoFile(AbsoluteFilePathAndName);
+			}
+			else
+			{
+				bError = true;
+			}
+		}
+		else
+		{
+			bError = true;
+		}
+
+		if (bError)
+		{
+			if (ErrorText.IsEmpty())
+			{
+				ErrorText = LOCTEXT("DownloadLatestFromLocalizationServiceFileProcessError", "An error occured when processing the file downloaded from the Localization Service.");
+			}
+		}
+	}
+	else
+	{
+		bError = true;
+		if (ErrorText.IsEmpty())
+		{
+			ErrorText = LOCTEXT("DownloadLatestFromLocalizationServiceDownloadError", "An error occured while downloading the file from the Localization Service.");
+		}
+	}
+
+	GWarn->StatusUpdate(100, 100, LOCTEXT("ImportFromLocalizationServiceFinished", "Import from Localization Service Complete!"));
+	GWarn->EndSlowTask();
+
+	if (bError || !ErrorText.IsEmpty())
+	{
+		if (ErrorText.IsEmpty())
+		{
+			ErrorText = LOCTEXT("DownloadLatestFromLocalizationServiceUnspecifiedError", "An unspecified error occured when trying download and import from the Localization Service.");
+		}
+
+		FMessageLog TranslationEditorMessageLog("TranslationEditor");
+		TranslationEditorMessageLog.Error(ErrorText);
+		TranslationEditorMessageLog.Notify(ErrorText);
+	}
+}
+
 void FTranslationEditor::ExportToPortableObjectFormat_Execute()
 {
 	const FString PortableObjectFileDescription = LOCTEXT("PortableObjectFileDescription", "Portable Object File").ToString();
@@ -1296,7 +1395,6 @@ void FTranslationEditor::ImportFromPortableObjectFormat_Execute()
 		DefaultPath = LastImportFilePath;
 	}
 	TArray<FString> OpenFilenames;
-	bool bHadError = false;
 	IDesktopPlatform* DesktopPlatform = FDesktopPlatformModule::Get();
 
 	bool bOpened = false;
@@ -1321,92 +1419,85 @@ void FTranslationEditor::ImportFromPortableObjectFormat_Execute()
 			);
 	}
 
-	if (bOpened)
+	if (bOpened && OpenFilenames.Num() > 0)
 	{
-		// Source path needs to be relative to Engine or Game directory
-		FString ManifestFullPath = FPaths::ConvertRelativePathToFull(ManifestFilePath);
-		FString EngineFullPath = FPaths::ConvertRelativePathToFull(FPaths::EngineContentDir());
-		FString FolderRelativeManifestPath = ManifestFilePath;
-		bool IsEngineManifest = false;
-		if (ManifestFullPath.StartsWith(EngineFullPath))
-		{
-			IsEngineManifest = true;
-		}
-		if (IsEngineManifest)
-		{
-			FPaths::MakePathRelativeTo(FolderRelativeManifestPath, *FPaths::EngineDir());
-		}
-		else
-		{
-			FPaths::MakePathRelativeTo(FolderRelativeManifestPath, *FPaths::GameDir());
-		}
-
-		// Paths like "Content/Localization/Editor/" are not considered relative for some reason, so add "./" if neccessary
-		if (!FPaths::IsRelative(FolderRelativeManifestPath))
-		{
-			FolderRelativeManifestPath = FString(TEXT("./")) + FolderRelativeManifestPath;
-		}
-
-		UInternationalizationExportSettings* ImportSettings = NewObject<UInternationalizationExportSettings>();
-		ImportSettings->CulturesToGenerate.Empty();
-		ImportSettings->CulturesToGenerate.Add(FPaths::GetBaseFilename(FPaths::GetPath(ArchiveFilePath)));
-		ImportSettings->CommandletClass = "InternationalizationExport";
-		ImportSettings->DestinationPath = FPaths::GetPath(FolderRelativeManifestPath);
-		ImportSettings->ManifestName = FPaths::GetBaseFilename(ManifestFilePath) + ".manifest";
-		ImportSettings->ArchiveName = FPaths::GetBaseFilename(ManifestFilePath) + ".archive";
-		ImportSettings->bExportLoc = false;
-		ImportSettings->bImportLoc = true;
-		ImportSettings->bUseCultureDirectory = false;
+		FString FileToImport = OpenFilenames[0];
 
-		ImportSettings->SourcePath = DefaultPath / FPaths::GetBaseFilename(ManifestFilePath);
+		ImportFromPoFile(FileToImport);
 
-		if (OpenFilenames.Num() > 0)
-		{
-			ImportSettings->SourcePath = FPaths::GetPath(OpenFilenames[0]);
-			ImportSettings->PortableObjectName = FPaths::GetCleanFilename(OpenFilenames[0]);
-			LastImportFilePath = FPaths::GetPath(OpenFilenames[0]);
-		}
+	}
+}
 
-		// Write translation data first to ensure all changes are exported
-		bHadError = !(DataManager->WriteTranslationData(true));
+void FTranslationEditor::ImportFromPoFile(FString FileToImport)
+{
+	// Source path needs to be relative to Engine or Game directory
+	FString ManifestFullPath = FPaths::ConvertRelativePathToFull(ManifestFilePath);
+	FString EngineFullPath = FPaths::ConvertRelativePathToFull(FPaths::EngineContentDir());
+	FString FolderRelativeManifestPath = ManifestFilePath;
+	bool IsEngineManifest = false;
+	if (ManifestFullPath.StartsWith(EngineFullPath))
+	{
+		IsEngineManifest = true;
+	}
+	if (IsEngineManifest)
+	{
+		FPaths::MakePathRelativeTo(FolderRelativeManifestPath, *FPaths::EngineDir());
+	}
+	else
+	{
+		FPaths::MakePathRelativeTo(FolderRelativeManifestPath, *FPaths::GameDir());
+	}
 
-		if (!bHadError)
-		{
-			GWarn->BeginSlowTask(LOCTEXT("ImportingInternationalization", "Importing Internationalization Data..."), true);
+	// Paths like "Content/Localization/Editor/" are not considered relative for some reason, so add "./" if neccessary
+	if (!FPaths::IsRelative(FolderRelativeManifestPath))
+	{
+		FolderRelativeManifestPath = FString(TEXT("./")) + FolderRelativeManifestPath;
+	}
 
-			// Write these settings to a temporary config file that the Internationalization Export Commandlet will read
-			FString TempConfigFilepath = FPaths::GameSavedDir() / "Config" / "InternationalizationExport.ini";
-			ImportSettings->SaveConfig(CPF_Config, *TempConfigFilepath);
+	UInternationalizationExportSettings* ImportSettings = NewObject<UInternationalizationExportSettings>();
+	ImportSettings->CulturesToGenerate.Empty();
+	ImportSettings->CulturesToGenerate.Add(FPaths::GetBaseFilename(FPaths::GetPath(ArchiveFilePath)));
+	ImportSettings->CommandletClass = "InternationalizationExport";
+	ImportSettings->DestinationPath = FPaths::GetPath(FolderRelativeManifestPath);
+	ImportSettings->ManifestName = FPaths::GetBaseFilename(ManifestFilePath) + ".manifest";
+	ImportSettings->ArchiveName = FPaths::GetBaseFilename(ManifestFilePath) + ".archive";
+	ImportSettings->bExportLoc = false;
+	ImportSettings->bImportLoc = true;
+	ImportSettings->bUseCultureDirectory = false;
+
+	ImportSettings->SourcePath = FPaths::GetPath(FileToImport);
+	ImportSettings->PortableObjectName = FPaths::GetCleanFilename(FileToImport);
+	LastImportFilePath = FPaths::GetPath(FileToImport);
+
+	// Write translation data first to ensure all changes are exported
+	bool bHadError = !(DataManager->WriteTranslationData(true));
+
+	if (!bHadError)
+	{
+		GWarn->BeginSlowTask(LOCTEXT("ImportingInternationalization", "Importing Internationalization Data..."), true);
 
-			// Using .ini config saving means these settings will be saved in the GetClass()->GetPathName() section
-			TArray<FString> ConfigSections;
-			ConfigSections.Add(ImportSettings->GetClass()->GetPathName());
-			FMessageLog TranslationEditorMessageLog("TranslationEditor");
+		// Write these settings to a temporary config file that the Internationalization Export Commandlet will read
+		FString TempConfigFilepath = FPaths::GameSavedDir() / "Config" / "InternationalizationExport.ini";
+		ImportSettings->SaveConfig(CPF_Config, *TempConfigFilepath);
 
-			for (FString& ConfigSection : ConfigSections)
-			{
-				// Spawn the LocalizationExport commandlet, and run its log output back into ours
-				FString AppURL = FPlatformProcess::ExecutableName(true);
-				FString Parameters = FString("-run=InternationalizationExport -config=") + TempConfigFilepath + " -section=" + ConfigSection;
+		// Using .ini config saving means these settings will be saved in the GetClass()->GetPathName() section
+		TArray<FString> ConfigSections;
+		ConfigSections.Add(ImportSettings->GetClass()->GetPathName());
+		FMessageLog TranslationEditorMessageLog("TranslationEditor");
 
-				void* WritePipe;
-				void* ReadPipe;
-				FPlatformProcess::CreatePipe(ReadPipe, WritePipe);
-				FProcHandle ProcessHandle = FPlatformProcess::CreateProc(*AppURL, *Parameters, false, false, false, NULL, 0, NULL, WritePipe);
+		for (FString& ConfigSection : ConfigSections)
+		{
+			// Spawn the LocalizationExport commandlet, and run its log output back into ours
+			FString AppURL = FPlatformProcess::ExecutableName(true);
+			FString Parameters = FString("-run=InternationalizationExport -config=") + TempConfigFilepath + " -section=" + ConfigSection;
 
-				while (FPlatformProcess::IsProcRunning(ProcessHandle))
-				{
-					FString NewLine = FPlatformProcess::ReadPipe(ReadPipe);
-					if (NewLine.Len() > 0)
-					{
-						UE_LOG(LocalizationExport, Log, TEXT("%s"), *NewLine);
-						FFormatNamedArguments Arguments;
-						Arguments.Add(TEXT("LogMessage"), FText::FromString(NewLine));
-						TranslationEditorMessageLog.Info(FText::Format(LOCTEXT("LocalizationImportLog", "Localization Import Log: {LogMessage}"), Arguments));
-					}
+			void* WritePipe;
+			void* ReadPipe;
+			FPlatformProcess::CreatePipe(ReadPipe, WritePipe);
+			FProcHandle ProcessHandle = FPlatformProcess::CreateProc(*AppURL, *Parameters, false, false, false, NULL, 0, NULL, WritePipe);
 
-					FPlatformProcess::Sleep(0.25);
-				}
+			while (FPlatformProcess::IsProcRunning(ProcessHandle))
+			{
 				FString NewLine = FPlatformProcess::ReadPipe(ReadPipe);
 				if (NewLine.Len() > 0)
 				{
@@ -1417,39 +1508,49 @@ void FTranslationEditor::ImportFromPortableObjectFormat_Execute()
 				}
 
 				FPlatformProcess::Sleep(0.25);
-				FPlatformProcess::ClosePipe(ReadPipe, WritePipe);
-
-				int32 ReturnCode;
-				if (!FPlatformProcess::GetProcReturnCode(ProcessHandle, &ReturnCode))
-				{
-					bHadError = true;
-				}
-				else if (ReturnCode != 0)
-				{
-					bHadError = true;
-				}
+			}
+			FString NewLine = FPlatformProcess::ReadPipe(ReadPipe);
+			if (NewLine.Len() > 0)
+			{
+				UE_LOG(LocalizationExport, Log, TEXT("%s"), *NewLine);
+				FFormatNamedArguments Arguments;
+				Arguments.Add(TEXT("LogMessage"), FText::FromString(NewLine));
+				TranslationEditorMessageLog.Info(FText::Format(LOCTEXT("LocalizationImportLog", "Localization Import Log: {LogMessage}"), Arguments));
 			}
 
-			GWarn->EndSlowTask();
+			FPlatformProcess::Sleep(0.25);
+			FPlatformProcess::ClosePipe(ReadPipe, WritePipe);
 
-			if (bHadError)
+			int32 ReturnCode;
+			if (!FPlatformProcess::GetProcReturnCode(ProcessHandle, &ReturnCode))
 			{
-				TranslationEditorMessageLog.Error(LOCTEXT("FailedToExportLocalization", "Failed to export localization!"));
-				TranslationEditorMessageLog.Notify(LOCTEXT("FailedToExportLocalization", "Failed to export localization!"), EMessageSeverity::Info, true);
-				TranslationEditorMessageLog.Open(EMessageSeverity::Error);
+				bHadError = true;
 			}
-			else
+			else if (ReturnCode != 0)
 			{
-				DataManager->LoadFromArchive(DataManager->GetAllTranslationsArray(), true, true);
-
-				TabManager->InvokeTab(ChangedOnImportTabId);
-				ChangedOnImportPropertyTable->SetObjects((TArray<UObject*>&)DataManager->GetChangedOnImportArray());
-				// Need to re-add the columns we want to display
-				ChangedOnImportPropertyTable->AddColumn((TWeakObjectPtr<UProperty>)FindField<UProperty>(UTranslationUnit::StaticClass(), "Source"));
-				ChangedOnImportPropertyTable->AddColumn((TWeakObjectPtr<UProperty>)FindField<UProperty>(UTranslationUnit::StaticClass(), "TranslationBeforeImport"));
-				ChangedOnImportPropertyTable->AddColumn((TWeakObjectPtr<UProperty>)FindField<UProperty>(UTranslationUnit::StaticClass(), "Translation"));
+				bHadError = true;
 			}
 		}
+
+		GWarn->EndSlowTask();
+
+		if (bHadError)
+		{
+			TranslationEditorMessageLog.Error(LOCTEXT("FailedToExportLocalization", "Failed to export localization!"));
+			TranslationEditorMessageLog.Notify(LOCTEXT("FailedToExportLocalization", "Failed to export localization!"), EMessageSeverity::Info, true);
+			TranslationEditorMessageLog.Open(EMessageSeverity::Error);
+		}
+		else
+		{
+			DataManager->LoadFromArchive(DataManager->GetAllTranslationsArray(), true, true);
+
+			TabManager->InvokeTab(ChangedOnImportTabId);
+			ChangedOnImportPropertyTable->SetObjects((TArray<UObject*>&)DataManager->GetChangedOnImportArray());
+			// Need to re-add the columns we want to display
+			ChangedOnImportPropertyTable->AddColumn((TWeakObjectPtr<UProperty>)FindField<UProperty>(UTranslationUnit::StaticClass(), "Source"));
+			ChangedOnImportPropertyTable->AddColumn((TWeakObjectPtr<UProperty>)FindField<UProperty>(UTranslationUnit::StaticClass(), "TranslationBeforeImport"));
+			ChangedOnImportPropertyTable->AddColumn((TWeakObjectPtr<UProperty>)FindField<UProperty>(UTranslationUnit::StaticClass(), "Translation"));
+		}
 	}
 }
 
diff --git a/Engine/Source/Editor/TranslationEditor/Private/TranslationEditor.h b/Engine/Source/Editor/TranslationEditor/Private/TranslationEditor.h
--- a/Engine/Source/Editor/TranslationEditor/Private/TranslationEditor.h
+++ b/Engine/Source/Editor/TranslationEditor/Private/TranslationEditor.h
@@ -6,6 +6,7 @@
 #include "ITranslationEditor.h"
 #include "CustomFontColumn.h"
 #include "TranslationUnit.h"
+#include "ILocalizationServiceProvider.h"
 
 class TRANSLATIONEDITOR_API FTranslationEditor :  public ITranslationEditor
 {
@@ -15,9 +16,9 @@ public:
 	/**
 	 *	Creates a new FTranslationEditor and calls Initialize
 	 */
-	static TSharedRef< FTranslationEditor > Create(TSharedRef< FTranslationDataManager > DataManager, const FString& InManifestFile, const FString& InArchiveFile)
+	static TSharedRef< FTranslationEditor > Create(TSharedRef< FTranslationDataManager > DataManager, const FString& InManifestFile, const FString& InArchiveFile, const FGuid& LocalizationTargetGuid)
 	{
-		TSharedRef< FTranslationEditor > TranslationEditor = MakeShareable(new FTranslationEditor(DataManager, InManifestFile, InArchiveFile));
+		TSharedRef< FTranslationEditor > TranslationEditor = MakeShareable(new FTranslationEditor(DataManager, InManifestFile, InArchiveFile, LocalizationTargetGuid));
 
 		// Some stuff that needs to use the "this" pointer is done in Initialize (because it can't be done in the constructor)
 		TranslationEditor->Initialize();
@@ -59,8 +60,8 @@ protected:
 
 private:
 
-	FTranslationEditor(TSharedRef< FTranslationDataManager > InDataManager, const FString& InManifestFile, const FString& InArchiveFile )
-		: ITranslationEditor(InManifestFile, InArchiveFile)
+	FTranslationEditor(TSharedRef< FTranslationDataManager > InDataManager, const FString& InManifestFile, const FString& InArchiveFile, const FGuid& InLocalizationTargetGuid)
+		: ITranslationEditor(InManifestFile, InArchiveFile, InLocalizationTargetGuid)
 	, DataManager(InDataManager)
 	, SourceFont(FEditorStyle::GetFontStyle( PropertyTableConstants::NormalFontStyle ))
 	, TranslationTargetFont(FEditorStyle::GetFontStyle( PropertyTableConstants::NormalFontStyle ))
@@ -158,15 +159,24 @@ private:
 	/** Update content when a new context selection is made */
 	void UpdateContextSelection();
 
-	/** Called when "Preview in Editor" is clicked for this asset */
+	/** Called when "Preview in Editor" is clicked for this Localization Target */
 	void PreviewAllTranslationsInEditor_Execute();
+	
+	/** Called when "Import from Localization Service" is clicked for this Localization Target */
+	void ImportLatestFromLocalizationService_Execute();
+
+	/** Callback for when the Localization Service operation started when "Import from Localization Service" was clicked finishes */
+	void DownloadLatestFromLocalizationServiceComplete(const FLocalizationServiceOperationRef& Operation, ELocalizationServiceOperationCommandResult::Type Result);
 
-	/** Called when "Export to .PO" is clicked for this asset */
+	/** Called when "Export to .PO" is clicked for this Localization Target */
 	void ExportToPortableObjectFormat_Execute();
 
-	/** Called when "Import from .PO" is clicked for this asset */
+	/** Called when "Import from .PO" is clicked for this Localization Target */
 	void ImportFromPortableObjectFormat_Execute();
 
+	/** Import from the specified .po file into this localization target */
+	void ImportFromPoFile(FString FileToImport);
+
 	/** Open the search tab */
 	void OpenSearchTab_Execute();
 
diff --git a/Engine/Source/Editor/TranslationEditor/Private/TranslationEditorMenu.cpp b/Engine/Source/Editor/TranslationEditor/Private/TranslationEditorMenu.cpp
--- a/Engine/Source/Editor/TranslationEditor/Private/TranslationEditorMenu.cpp
+++ b/Engine/Source/Editor/TranslationEditor/Private/TranslationEditorMenu.cpp
@@ -17,6 +17,7 @@ void FTranslationEditorMenu::FillTranslationMenu( FMenuBuilder& MenuBuilder/*, F
 		MenuBuilder.AddMenuEntry( FTranslationEditorCommands::Get().ChangeSourceFont );
 		MenuBuilder.AddMenuEntry( FTranslationEditorCommands::Get().ChangeTranslationTargetFont );
 		MenuBuilder.AddMenuEntry( FTranslationEditorCommands::Get().PreviewAllTranslationsInEditor );
+		MenuBuilder.AddMenuEntry( FTranslationEditorCommands::Get().ImportLatestFromLocalizationService );
 		MenuBuilder.AddMenuEntry( FTranslationEditorCommands::Get().ExportToPortableObjectFormat );
 		MenuBuilder.AddMenuEntry( FTranslationEditorCommands::Get().ImportFromPortableObjectFormat );
 		MenuBuilder.AddMenuEntry( FTranslationEditorCommands::Get().OpenSearchTab );
@@ -71,6 +72,8 @@ void FTranslationEditorMenu::SetupTranslationEditorToolbar( TSharedPtr< FExtende
 				FTranslationEditorCommands::Get().SaveTranslations, "SaveTranslations", TAttribute<FText>(), TAttribute<FText>(), FSlateIcon(FEditorStyle::GetStyleSetName(), "AssetEditor.SaveAsset"));
 			ToolbarBuilder.AddToolBarButton(
 				FTranslationEditorCommands::Get().PreviewAllTranslationsInEditor, "PreviewTranslationsInEditor", TAttribute<FText>(), TAttribute<FText>(), FSlateIcon(FEditorStyle::GetStyleSetName(), "TranslationEditor.PreviewInEditor")); 
+			ToolbarBuilder.AddToolBarButton(
+				FTranslationEditorCommands::Get().ImportLatestFromLocalizationService, "ImportLatestFromLocalizationService", TAttribute<FText>(), TAttribute<FText>(), FSlateIcon(FEditorStyle::GetStyleSetName(), "TranslationEditor.ImportLatestFromLocalizationService"));
 			ToolbarBuilder.AddToolBarButton(
 				FTranslationEditorCommands::Get().ExportToPortableObjectFormat, "ExportToPortableObjectFormat", TAttribute<FText>(), TAttribute<FText>(), FSlateIcon(FEditorStyle::GetStyleSetName(), "TranslationEditor.Export"));
 			ToolbarBuilder.AddToolBarButton(
@@ -97,7 +100,8 @@ void FTranslationEditorCommands::RegisterCommands()
 	UI_COMMAND( ChangeSourceFont, "Change Source Font", "Change the Font for the Source Lanugage", EUserInterfaceActionType::Button, FInputChord() );
 	UI_COMMAND( ChangeTranslationTargetFont, "Change Translation Font", "Change the Translation Target Language Font", EUserInterfaceActionType::Button, FInputChord() );
 	UI_COMMAND( SaveTranslations, "Save", "Saves the translations to file", EUserInterfaceActionType::Button, FInputChord() );
-	UI_COMMAND( PreviewAllTranslationsInEditor, "Preview in Editor", "Preview All Translations in the Editor UI", EUserInterfaceActionType::Button, FInputChord() );
+	UI_COMMAND( PreviewAllTranslationsInEditor, "Preview in Editor", "Preview All Translations in the Editor UI", EUserInterfaceActionType::Button, FInputChord());
+	UI_COMMAND( ImportLatestFromLocalizationService, "Import from Translation Service", "Download and Import Latest Translations From Localization Service. (Localization Service settings can be modified in the Localization Dashboard)", EUserInterfaceActionType::Button, FInputChord());
 	UI_COMMAND( ExportToPortableObjectFormat, "Export to .PO", "Export to Portable Object Format", EUserInterfaceActionType::Button, FInputChord() );
 	UI_COMMAND( ImportFromPortableObjectFormat, "Import from .PO", "Import from Portable Object Format", EUserInterfaceActionType::Button, FInputChord() );
 	UI_COMMAND( OpenSearchTab, "Search", "Search Source and Translation Strings", EUserInterfaceActionType::Button, FInputChord() );
diff --git a/Engine/Source/Editor/TranslationEditor/Private/TranslationEditorModule.cpp b/Engine/Source/Editor/TranslationEditor/Private/TranslationEditorModule.cpp
--- a/Engine/Source/Editor/TranslationEditor/Private/TranslationEditorModule.cpp
+++ b/Engine/Source/Editor/TranslationEditor/Private/TranslationEditorModule.cpp
@@ -39,13 +39,13 @@ void FTranslationEditorModule::ShutdownModule()
 #endif
 }
 
-TSharedRef<FTranslationEditor> FTranslationEditorModule::CreateTranslationEditor(const FString& ManifestFile, const FString& ArchiveFile)
+TSharedRef<FTranslationEditor> FTranslationEditorModule::CreateTranslationEditor(const FString& ManifestFile, const FString& ArchiveFile, const FGuid& LocalizationTargetGuid)
 {
 	TSharedRef< FTranslationDataManager > DataManager = MakeShareable( new FTranslationDataManager(ManifestFile, ArchiveFile) );
 
 	GWarn->BeginSlowTask(LOCTEXT("BuildingUserInterface", "Building Translation Editor UI..."), true);
 
-	TSharedRef< FTranslationEditor > NewTranslationEditor( FTranslationEditor::Create(DataManager, ManifestFile, ArchiveFile) );
+	TSharedRef< FTranslationEditor > NewTranslationEditor( FTranslationEditor::Create(DataManager, ManifestFile, ArchiveFile, LocalizationTargetGuid) );
 	NewTranslationEditor->InitTranslationEditor( EToolkitMode::Standalone, TSharedPtr<IToolkitHost>() );
 
 	GWarn->EndSlowTask();
diff --git a/Engine/Source/Editor/TranslationEditor/Public/ITranslationEditor.h b/Engine/Source/Editor/TranslationEditor/Public/ITranslationEditor.h
--- a/Engine/Source/Editor/TranslationEditor/Public/ITranslationEditor.h
+++ b/Engine/Source/Editor/TranslationEditor/Public/ITranslationEditor.h
@@ -14,13 +14,14 @@ class ITranslationEditor : public FAssetEditorToolkit
 
 public:
 
-	TRANSLATIONEDITOR_API static void OpenTranslationEditor(const FString& InManifestFile, const FString& InArchiveFile);
+	TRANSLATIONEDITOR_API static void OpenTranslationEditor(const FString& InManifestFile, const FString& InArchiveFile, const FGuid& LocalizationTargetGuid);
 
 	TRANSLATIONEDITOR_API static void OpenTranslationPicker();
 
-	ITranslationEditor(const FString& InManifestFile, const FString& InArchiveFile)
+	ITranslationEditor(const FString& InManifestFile, const FString& InArchiveFile, const FGuid& InLocalizationTargetGuid)
 		: ManifestFilePath(InManifestFile)
 		, ArchiveFilePath(InArchiveFile)
+		, LocalizationTargetGuid(InLocalizationTargetGuid)
 	{}
 
 	virtual bool OnRequestClose() override;
@@ -31,6 +32,8 @@ protected:
 	FString ManifestFilePath;
 	/** Name of the language we are translating to */
 	FString ArchiveFilePath;
+	/** GUID of the Localization Target we are editing */
+	FGuid LocalizationTargetGuid;
 
 private:
 
diff --git a/Engine/Source/Editor/TranslationEditor/Public/TranslationEditorMenu.h b/Engine/Source/Editor/TranslationEditor/Public/TranslationEditorMenu.h
--- a/Engine/Source/Editor/TranslationEditor/Public/TranslationEditorMenu.h
+++ b/Engine/Source/Editor/TranslationEditor/Public/TranslationEditorMenu.h
@@ -35,6 +35,9 @@ public:
 
 	/** Save the translations to file */
 	TSharedPtr<FUICommandInfo> PreviewAllTranslationsInEditor;
+	
+	/** Download and import the latest translations from the active Translation Service */
+	TSharedPtr<FUICommandInfo> ImportLatestFromLocalizationService;
 
 	/** Export to PortableObject format (.po) */
 	TSharedPtr<FUICommandInfo> ExportToPortableObjectFormat;
diff --git a/Engine/Source/Editor/TranslationEditor/Public/TranslationEditorModule.h b/Engine/Source/Editor/TranslationEditor/Public/TranslationEditorModule.h
--- a/Engine/Source/Editor/TranslationEditor/Public/TranslationEditorModule.h
+++ b/Engine/Source/Editor/TranslationEditor/Public/TranslationEditorModule.h
@@ -26,7 +26,7 @@ public:
 	 * 
 	 * @return	Interface to the new translation editor
 	 */
-	virtual TSharedRef<FTranslationEditor> CreateTranslationEditor( const FString& ManifestFile, const FString& ArchiveFile );
+	virtual TSharedRef<FTranslationEditor> CreateTranslationEditor(const FString& ManifestFile, const FString& ArchiveFile, const FGuid& LocalizationTargetGuid);
 
 	/** Gets the extensibility managers for outside entities to extend translation editor's menus and toolbars */
 	virtual TSharedPtr<FExtensibilityManager> GetMenuExtensibilityManager() override { return MenuExtensibilityManager; }
diff --git a/Engine/Source/Editor/TranslationEditor/TranslationEditor.Build.cs b/Engine/Source/Editor/TranslationEditor/TranslationEditor.Build.cs
--- a/Engine/Source/Editor/TranslationEditor/TranslationEditor.Build.cs
+++ b/Engine/Source/Editor/TranslationEditor/TranslationEditor.Build.cs
@@ -8,6 +8,8 @@ public class TranslationEditor : ModuleRules
 	{
 		PublicIncludePathModuleNames.Add("LevelEditor");
 		PublicIncludePathModuleNames.Add("WorkspaceMenuStructure");
+        
+        PrivateIncludePathModuleNames.Add("LocalizationService");
 
         PrivateIncludePathModuleNames.AddRange(
 			new string[] {
@@ -33,6 +35,7 @@ public class TranslationEditor : ModuleRules
 				"SourceControl",
                 "MessageLog",
                 "Documentation",
+                "LocalizationService",
 			}
 		);
 
@@ -48,7 +51,8 @@ public class TranslationEditor : ModuleRules
 			new string[] {
 				"WorkspaceMenuStructure",
 				"DesktopPlatform",
-				"SourceControl"
+				"SourceControl",
+                "LocalizationService",
 			}
 		);
 	}
diff --git a/Engine/Source/Editor/UnrealEd/Private/EditorEngine.cpp b/Engine/Source/Editor/UnrealEd/Private/EditorEngine.cpp
--- a/Engine/Source/Editor/UnrealEd/Private/EditorEngine.cpp
+++ b/Engine/Source/Editor/UnrealEd/Private/EditorEngine.cpp
@@ -20,6 +20,7 @@
 #include "ScopedTransaction.h"
 
 #include "ISourceControlModule.h"
+#include "ILocalizationServiceModule.h"
 #include "PackageBackup.h"
 #include "LevelUtils.h"
 #include "Layers/Layers.h"
@@ -1490,6 +1491,7 @@ void UEditorEngine::Tick( float DeltaSeconds, bool bIdleMode )
 	}
 
 	ISourceControlModule::Get().Tick();
+	ILocalizationServiceModule::Get().Tick();
 
 	if( FSlateThrottleManager::Get().IsAllowingExpensiveTasks() )
 	{
diff --git a/Engine/Source/Editor/UnrealEd/UnrealEd.Build.cs b/Engine/Source/Editor/UnrealEd/UnrealEd.Build.cs
--- a/Engine/Source/Editor/UnrealEd/UnrealEd.Build.cs
+++ b/Engine/Source/Editor/UnrealEd/UnrealEd.Build.cs
@@ -56,6 +56,7 @@ public class UnrealEd : ModuleRules
 				"Landscape",
                 "Niagara",
 				"SizeMap",
+                "LocalizationService",
 			}
 		);
 
@@ -130,6 +131,7 @@ public class UnrealEd : ModuleRules
                 "VectorVM",
 				"TreeMap",
                 "MaterialUtilities",
+                "LocalizationService",
 			}
 		);
 
