Translation Editor: Remove check() that was getting hit when the .manifest or .archive failed to load, then you tried to import anyway.  Now just displays "failed to import" error message without crashing.

diff --git a/Engine/Source/Editor/TranslationEditor/Private/ITranslationEditor.cpp b/Engine/Source/Editor/TranslationEditor/Private/ITranslationEditor.cpp
--- a/Engine/Source/Editor/TranslationEditor/Private/ITranslationEditor.cpp
+++ b/Engine/Source/Editor/TranslationEditor/Private/ITranslationEditor.cpp
@@ -14,8 +14,15 @@ void ITranslationEditor::OpenTranslationEditor(const FString& InManifestFile, co
 	if (!OpenTranslationEditors.Contains(InArchiveFile))
 	{
 		FTranslationEditorModule& TranslationEditorModule = FModuleManager::LoadModuleChecked<FTranslationEditorModule>("TranslationEditor");
-		ITranslationEditor* NewTranslationEditor = (ITranslationEditor*)&(TranslationEditorModule.CreateTranslationEditor(InManifestFile, InArchiveFile, LocalizationTargetGuid).Get());
+		bool bLoadedSuccessfully = false;
+		ITranslationEditor* NewTranslationEditor = (ITranslationEditor*)&(TranslationEditorModule.CreateTranslationEditor(InManifestFile, InArchiveFile, LocalizationTargetGuid, bLoadedSuccessfully).Get());
+
 		NewTranslationEditor->RegisterTranslationEditor();
+
+		if (!bLoadedSuccessfully)
+		{
+			NewTranslationEditor->CloseWindow();
+		}
 	}
 	else
 	{
diff --git a/Engine/Source/Editor/TranslationEditor/Private/TranslationDataManager.cpp b/Engine/Source/Editor/TranslationEditor/Private/TranslationDataManager.cpp
--- a/Engine/Source/Editor/TranslationEditor/Private/TranslationDataManager.cpp
+++ b/Engine/Source/Editor/TranslationEditor/Private/TranslationDataManager.cpp
@@ -21,6 +21,7 @@ DEFINE_LOG_CATEGORY_STATIC(LogTranslationEditor, Log, All);
 FTranslationDataManager::FTranslationDataManager( const FString& InManifestFilePath, const FString& InArchiveFilePath )
 : ManifestFilePath(InManifestFilePath)
 , ArchiveFilePath(InArchiveFilePath)
+, bLoadedSuccessfully(true)
 {
 	GWarn->BeginSlowTask(LOCTEXT("LoadingTranslationData", "Loading Translation Data..."), true);
 	TArray<UTranslationUnit*> TranslationUnits;
@@ -33,6 +34,7 @@ FTranslationDataManager::FTranslationDataManager( const FString& InManifestFileP
 
 		if (ManifestEntriesCount < 1)
 		{
+			bLoadedSuccessfully = false;
 			FFormatNamedArguments Arguments;
 			Arguments.Add( TEXT("ManifestFilePath"), FText::FromString(ManifestFilePath) );
 			Arguments.Add( TEXT("ManifestEntriesCount"), FText::AsNumber(ManifestEntriesCount) );
@@ -80,6 +82,7 @@ FTranslationDataManager::FTranslationDataManager( const FString& InManifestFileP
 		}
 		else  // ArchivePtr.IsValid() is false
 		{
+			bLoadedSuccessfully = false;
 			FFormatNamedArguments Arguments;
 			Arguments.Add( TEXT("ArchiveFilePath"), FText::FromString(ArchiveFilePath) );
 			FMessageLog TranslationEditorMessageLog("TranslationEditor");
@@ -90,6 +93,7 @@ FTranslationDataManager::FTranslationDataManager( const FString& InManifestFileP
 	}
 	else  // ManifestAtHeadRevisionPtr.IsValid() is false
 	{
+		bLoadedSuccessfully = false;
 		FFormatNamedArguments Arguments;
 		Arguments.Add( TEXT("ManifestFilePath"), FText::FromString(ManifestFilePath) );
 		FMessageLog TranslationEditorMessageLog("TranslationEditor");
@@ -168,64 +172,69 @@ TSharedPtr<FJsonObject> FTranslationDataManager::ReadJSONTextFile(const FString&
 
 bool FTranslationDataManager::WriteTranslationData(bool bForceWrite /*= false*/)
 {
-	check (ArchivePtr.IsValid());
-	TSharedRef< FInternationalizationArchive > Archive = ArchivePtr.ToSharedRef();
+	bool bSuccess = false;
 
-	bool bNeedsWrite = false;
-
-	for (UTranslationUnit* TranslationUnit : Untranslated)
+	// If the archive hasn't been loaded correctly, don't try and write anything
+	if (ArchivePtr.IsValid())
 	{
-		if (TranslationUnit != nullptr)
+		TSharedRef< FInternationalizationArchive > Archive = ArchivePtr.ToSharedRef();
+
+		bool bNeedsWrite = false;
+
+		for (UTranslationUnit* TranslationUnit : Untranslated)
 		{
-			const FLocItem SearchSource(TranslationUnit->Source);
-			FString OldTranslation = Archive->FindEntryBySource(TranslationUnit->Namespace, SearchSource, nullptr)->Translation.Text;
-			FString TranslationToWrite = TranslationUnit->Translation;
-			if (!TranslationToWrite.Equals(OldTranslation))
+			if (TranslationUnit != nullptr)
 			{
-				Archive->SetTranslation(TranslationUnit->Namespace, TranslationUnit->Source, TranslationToWrite, nullptr);
-				bNeedsWrite = true;
+				const FLocItem SearchSource(TranslationUnit->Source);
+				FString OldTranslation = Archive->FindEntryBySource(TranslationUnit->Namespace, SearchSource, nullptr)->Translation.Text;
+				FString TranslationToWrite = TranslationUnit->Translation;
+				if (!TranslationToWrite.Equals(OldTranslation))
+				{
+					Archive->SetTranslation(TranslationUnit->Namespace, TranslationUnit->Source, TranslationToWrite, nullptr);
+					bNeedsWrite = true;
+				}
 			}
 		}
-	}
 
-	for (UTranslationUnit* TranslationUnit : Review)
-	{
-		if (TranslationUnit != nullptr)
+		for (UTranslationUnit* TranslationUnit : Review)
 		{
-			const FLocItem SearchSource(TranslationUnit->Source);
-			FString OldTranslation = Archive->FindEntryBySource(TranslationUnit->Namespace, SearchSource, nullptr)->Translation.Text;
-			FString TranslationToWrite = TranslationUnit->Translation;
-			if (TranslationUnit->HasBeenReviewed && !TranslationToWrite.Equals(OldTranslation))
+			if (TranslationUnit != nullptr)
 			{
-				Archive->SetTranslation(TranslationUnit->Namespace, TranslationUnit->Source, TranslationToWrite, nullptr);
-				bNeedsWrite = true;
+				const FLocItem SearchSource(TranslationUnit->Source);
+				FString OldTranslation = Archive->FindEntryBySource(TranslationUnit->Namespace, SearchSource, nullptr)->Translation.Text;
+				FString TranslationToWrite = TranslationUnit->Translation;
+				if (TranslationUnit->HasBeenReviewed && !TranslationToWrite.Equals(OldTranslation))
+				{
+					Archive->SetTranslation(TranslationUnit->Namespace, TranslationUnit->Source, TranslationToWrite, nullptr);
+					bNeedsWrite = true;
+				}
 			}
 		}
-	}
 
-	for (UTranslationUnit* TranslationUnit : Complete)
-	{
-		if (TranslationUnit != nullptr)
+		for (UTranslationUnit* TranslationUnit : Complete)
 		{
-			const FLocItem SearchSource(TranslationUnit->Source);
-			FString OldTranslation = Archive->FindEntryBySource(TranslationUnit->Namespace, SearchSource, nullptr)->Translation.Text;
-			FString TranslationToWrite = TranslationUnit->Translation;
-			if (!TranslationToWrite.Equals(OldTranslation))
+			if (TranslationUnit != nullptr)
 			{
-				Archive->SetTranslation(TranslationUnit->Namespace, TranslationUnit->Source, TranslationToWrite, nullptr);
-				bNeedsWrite = true;
+				const FLocItem SearchSource(TranslationUnit->Source);
+				FString OldTranslation = Archive->FindEntryBySource(TranslationUnit->Namespace, SearchSource, nullptr)->Translation.Text;
+				FString TranslationToWrite = TranslationUnit->Translation;
+				if (!TranslationToWrite.Equals(OldTranslation))
+				{
+					Archive->SetTranslation(TranslationUnit->Namespace, TranslationUnit->Source, TranslationToWrite, nullptr);
+					bNeedsWrite = true;
+				}
 			}
 		}
-	}
 
-	bool bSuccess = true;
+		bSuccess = true;
 
-	if (bForceWrite || bNeedsWrite)
-	{
-		TSharedRef<FJsonObject> FinalArchiveJsonObj = MakeShareable( new FJsonObject );
-		ArchiveSerializer.SerializeArchive( Archive, FinalArchiveJsonObj );
+		if (bForceWrite || bNeedsWrite)
+		{
+			TSharedRef<FJsonObject> FinalArchiveJsonObj = MakeShareable(new FJsonObject);
+			ArchiveSerializer.SerializeArchive(Archive, FinalArchiveJsonObj);
 
-		bSuccess = WriteJSONToTextFile( FinalArchiveJsonObj, ArchiveFilePath );
+			bSuccess = WriteJSONToTextFile(FinalArchiveJsonObj, ArchiveFilePath);
+		}
 	}
 
 	return bSuccess;
diff --git a/Engine/Source/Editor/TranslationEditor/Private/TranslationDataManager.h b/Engine/Source/Editor/TranslationEditor/Private/TranslationDataManager.h
--- a/Engine/Source/Editor/TranslationEditor/Private/TranslationDataManager.h
+++ b/Engine/Source/Editor/TranslationEditor/Private/TranslationDataManager.h
@@ -76,7 +76,14 @@ public:
 	/** Save the specified translations */
 	static bool SaveSelectedTranslations(TArray<UTranslationUnit*> TranslationUnitsToSave);
 
+	/** Whether or not we successfully loaded the .manifest and .archive */
+	bool GetLoadedSuccessfully()
+	{
+		return bLoadedSuccessfully;
+	}
+
 private:
+
 	/** Read text file into a JSON file */
 	static TSharedPtr<FJsonObject> ReadJSONTextFile( const FString& InFilePath ) ;
 
@@ -124,4 +131,7 @@ private:
 
 	/** Files that are already checked out from Perforce **/
 	TArray<FString> CheckedOutFiles;
+
+	/** Whether or not we successfully loaded the .manifest and .archive */
+	bool bLoadedSuccessfully;
 };
diff --git a/Engine/Source/Editor/TranslationEditor/Private/TranslationEditor.cpp b/Engine/Source/Editor/TranslationEditor/Private/TranslationEditor.cpp
--- a/Engine/Source/Editor/TranslationEditor/Private/TranslationEditor.cpp
+++ b/Engine/Source/Editor/TranslationEditor/Private/TranslationEditor.cpp
@@ -1469,6 +1469,8 @@ void FTranslationEditor::ImportFromPoFile(FString FileToImport)
 	ImportSettings->PortableObjectName = FPaths::GetCleanFilename(FileToImport);
 	LastImportFilePath = FPaths::GetPath(FileToImport);
 
+	FMessageLog TranslationEditorMessageLog("TranslationEditor");
+
 	// Write translation data first to ensure all changes are exported
 	bool bHadError = !(DataManager->WriteTranslationData(true));
 
@@ -1483,7 +1485,6 @@ void FTranslationEditor::ImportFromPoFile(FString FileToImport)
 		// Using .ini config saving means these settings will be saved in the GetClass()->GetPathName() section
 		TArray<FString> ConfigSections;
 		ConfigSections.Add(ImportSettings->GetClass()->GetPathName());
-		FMessageLog TranslationEditorMessageLog("TranslationEditor");
 
 		for (FString& ConfigSection : ConfigSections)
 		{
@@ -1534,13 +1535,7 @@ void FTranslationEditor::ImportFromPoFile(FString FileToImport)
 
 		GWarn->EndSlowTask();
 
-		if (bHadError)
-		{
-			TranslationEditorMessageLog.Error(LOCTEXT("FailedToExportLocalization", "Failed to export localization!"));
-			TranslationEditorMessageLog.Notify(LOCTEXT("FailedToExportLocalization", "Failed to export localization!"), EMessageSeverity::Info, true);
-			TranslationEditorMessageLog.Open(EMessageSeverity::Error);
-		}
-		else
+		if (!bHadError)
 		{
 			DataManager->LoadFromArchive(DataManager->GetAllTranslationsArray(), true, true);
 
@@ -1552,6 +1547,14 @@ void FTranslationEditor::ImportFromPoFile(FString FileToImport)
 			ChangedOnImportPropertyTable->AddColumn((TWeakObjectPtr<UProperty>)FindField<UProperty>(UTranslationUnit::StaticClass(), "Translation"));
 		}
 	}
+
+	if (bHadError)
+	{
+		FText FailedToImport = LOCTEXT("FailedToImportLocalization", "Failed to import localization data!");
+		TranslationEditorMessageLog.Error(FailedToImport);
+		TranslationEditorMessageLog.Notify(FailedToImport, EMessageSeverity::Info, true);
+		TranslationEditorMessageLog.Open(EMessageSeverity::Error);
+	}
 }
 
 void FTranslationEditor::OnFilterTextChanged(const FText& InFilterText)
diff --git a/Engine/Source/Editor/TranslationEditor/Private/TranslationEditorModule.cpp b/Engine/Source/Editor/TranslationEditor/Private/TranslationEditorModule.cpp
--- a/Engine/Source/Editor/TranslationEditor/Private/TranslationEditorModule.cpp
+++ b/Engine/Source/Editor/TranslationEditor/Private/TranslationEditorModule.cpp
@@ -39,14 +39,15 @@ void FTranslationEditorModule::ShutdownModule()
 #endif
 }
 
-TSharedRef<FTranslationEditor> FTranslationEditorModule::CreateTranslationEditor(const FString& ManifestFile, const FString& ArchiveFile, const FGuid& LocalizationTargetGuid)
+TSharedRef<FTranslationEditor> FTranslationEditorModule::CreateTranslationEditor(const FString& ManifestFile, const FString& ArchiveFile, const FGuid& LocalizationTargetGuid, bool& OutLoadedSuccessfully)
 {
 	TSharedRef< FTranslationDataManager > DataManager = MakeShareable( new FTranslationDataManager(ManifestFile, ArchiveFile) );
+	OutLoadedSuccessfully = DataManager->GetLoadedSuccessfully();
 
 	GWarn->BeginSlowTask(LOCTEXT("BuildingUserInterface", "Building Translation Editor UI..."), true);
 
-	TSharedRef< FTranslationEditor > NewTranslationEditor( FTranslationEditor::Create(DataManager, ManifestFile, ArchiveFile, LocalizationTargetGuid) );
-	NewTranslationEditor->InitTranslationEditor( EToolkitMode::Standalone, TSharedPtr<IToolkitHost>() );
+	TSharedRef< FTranslationEditor > NewTranslationEditor(FTranslationEditor::Create(DataManager, ManifestFile, ArchiveFile, LocalizationTargetGuid));
+	NewTranslationEditor->InitTranslationEditor(EToolkitMode::Standalone, TSharedPtr<IToolkitHost>());
 
 	GWarn->EndSlowTask();
 
diff --git a/Engine/Source/Editor/TranslationEditor/Public/TranslationEditorModule.h b/Engine/Source/Editor/TranslationEditor/Public/TranslationEditorModule.h
--- a/Engine/Source/Editor/TranslationEditor/Public/TranslationEditorModule.h
+++ b/Engine/Source/Editor/TranslationEditor/Public/TranslationEditorModule.h
@@ -23,10 +23,12 @@ public:
 	 *
 	 * @param ProjectName				Name of the project to translate
 	 * @param TranslationTargetLanguage	Language to translate to
+	 * @param LocalizationTargetGuid	Unique ID for this Localization Target
+	 * @param OutLoadedSuccessfully		Whether or not the translation editor was able to load successfully from the .manifest and .archive files
 	 * 
 	 * @return	Interface to the new translation editor
 	 */
-	virtual TSharedRef<FTranslationEditor> CreateTranslationEditor(const FString& ManifestFile, const FString& ArchiveFile, const FGuid& LocalizationTargetGuid);
+	virtual TSharedRef<FTranslationEditor> CreateTranslationEditor(const FString& ManifestFile, const FString& ArchiveFile, const FGuid& LocalizationTargetGuid, bool& OutLoadedSuccessfully);
 
 	/** Gets the extensibility managers for outside entities to extend translation editor's menus and toolbars */
 	virtual TSharedPtr<FExtensibilityManager> GetMenuExtensibilityManager() override { return MenuExtensibilityManager; }
