Resolved shadowed variable warning in TranslationEditor.cpp.

diff --git a/Engine/Source/Editor/TranslationEditor/Private/TranslationEditor.cpp b/Engine/Source/Editor/TranslationEditor/Private/TranslationEditor.cpp
--- a/Engine/Source/Editor/TranslationEditor/Private/TranslationEditor.cpp
+++ b/Engine/Source/Editor/TranslationEditor/Private/TranslationEditor.cpp
@@ -1274,8 +1274,6 @@ void FTranslationEditor::ExportToPortableObjectFormat_Execute()
 	{
 		LastExportFilePath = FPaths::GetPath(SaveFilenames[0]);
 
-		const FString CultureToEdit = FPaths::GetBaseFilename(FPaths::GetPath(ArchiveFilePath));
-
 		// Write translation data first to ensure all changes are exported
 		if (DataManager->WriteTranslationData() && ParentWindow.IsValid() && SaveFilenames.Num() > 0)
 		{
