Fixed order-of-init warning in TranslationDataManager

diff --git a/Engine/Source/Editor/TranslationEditor/Private/TranslationDataManager.cpp b/Engine/Source/Editor/TranslationEditor/Private/TranslationDataManager.cpp
--- a/Engine/Source/Editor/TranslationEditor/Private/TranslationDataManager.cpp
+++ b/Engine/Source/Editor/TranslationEditor/Private/TranslationDataManager.cpp
@@ -24,10 +24,10 @@ DEFINE_LOG_CATEGORY_STATIC(LogTranslationEditor, Log, All);
 
 
 FTranslationDataManager::FTranslationDataManager( const FString& InManifestFilePath, const FString& InNativeArchiveFilePath, const FString& InArchiveFilePath )
-	: AssociatedLocalizationTarget(nullptr)
-	, OpenedManifestFilePath(InManifestFilePath)
+	: OpenedManifestFilePath(InManifestFilePath)
 	, NativeArchiveFilePath(InNativeArchiveFilePath)
 	, OpenedArchiveFilePath(InArchiveFilePath)
+	, AssociatedLocalizationTarget(nullptr)
 	, bLoadedSuccessfully(true)
 {
 	Initialize();
