Translation Picker: Small fix to prevent situation where text that was not ready for translation was allowed to be translated.

diff --git a/Engine/Source/Editor/TranslationEditor/Private/TranslationPickerEditWindow.cpp b/Engine/Source/Editor/TranslationEditor/Private/TranslationPickerEditWindow.cpp
--- a/Engine/Source/Editor/TranslationEditor/Private/TranslationPickerEditWindow.cpp
+++ b/Engine/Source/Editor/TranslationEditor/Private/TranslationPickerEditWindow.cpp
@@ -206,13 +206,12 @@ FReply STranslationPickerEditWindow::Close()
 
 FReply STranslationPickerEditWindow::SaveAllAndClose()
 {
-	
 	TArray<UTranslationUnit*> TempArray;
 
 	for (TSharedRef<STranslationPickerEditWidget> EditWidget : EditWidgets)
 	{
 		UTranslationUnit* TranslationUnit = EditWidget->GetTranslationUnitWithAnyChanges();
-		if (TranslationUnit != nullptr)
+		if (TranslationUnit != nullptr && EditWidget->CanSave())
 		{
 			TempArray.Add(TranslationUnit);
 		}
@@ -264,7 +263,7 @@ void STranslationPickerEditWidget::Construct(const FArguments& InArgs)
 	FString CurrentCultureName = I18N.GetCurrentCulture()->GetName();
 
 	// Make sure locres path matches the current culture (can change if using the culture=__ console command)
-	if (CurrentCultureName != LocResCultureName)
+	if (!LocResCultureName.IsEmpty() && CurrentCultureName != LocResCultureName)
 	{
 		FString BasePath = FPaths::GetPath(ArchiveFilePath);
 		LocresFullPath = BasePath / CurrentCultureName / ManifestAndArchiveNameString + ".locres";
@@ -289,7 +288,7 @@ void STranslationPickerEditWidget::Construct(const FArguments& InArgs)
 	TranslationUnit->LocresPath = LocresFullPath;
 
 	// Can only save if we have all the required information
-	bool bHasRequiredLocalizationInfo = NamespaceString.IsSet() && SourceString != nullptr && LocresFullPath.Len() > 0;
+	bHasRequiredLocalizationInfoForSaving = NamespaceString.IsSet() && SourceString != nullptr && LocresFullPath.Len() > 0;
 
 	TSharedPtr<SGridPanel> GridPanel;
 	TSharedRef<SGridPanel> LocalizationInfoAndSaveButtonSlot = SNew(SGridPanel).FillColumn(2,1);
@@ -320,7 +319,8 @@ void STranslationPickerEditWidget::Construct(const FArguments& InArgs)
 					.HAlign(HAlign_Right)
 				[
 					SNew(SVerticalBox)
-					.Visibility(!bHasRequiredLocalizationInfo && SourceString->Equals(TranslationString) ? EVisibility::Collapsed : EVisibility::Visible)
+					// Hide translation if we don't have necessary information to modify, and is same as source
+					.Visibility(!bHasRequiredLocalizationInfoForSaving && SourceString->Equals(TranslationString) ? EVisibility::Collapsed : EVisibility::Visible)
 					+ SVerticalBox::Slot()
 					[
 						SNew(STextBlock)
@@ -341,11 +341,12 @@ void STranslationPickerEditWidget::Construct(const FArguments& InArgs)
 					.Padding(FMargin(5))
 					[
 						SNew(SVerticalBox)
-						.Visibility(!bHasRequiredLocalizationInfo && SourceString->Equals(TranslationString) ? EVisibility::Collapsed : EVisibility::Visible)
+						// Hide translation if we don't have necessary information to modify, and is same as source
+						.Visibility(!bHasRequiredLocalizationInfoForSaving && SourceString->Equals(TranslationString) ? EVisibility::Collapsed : EVisibility::Visible)
 						+ SVerticalBox::Slot()
 						[
 							SAssignNew(TextBox, SMultiLineEditableTextBox)
-							.IsEnabled(bAllowEditing && bHasRequiredLocalizationInfo)
+							.IsEnabled(bAllowEditing && bHasRequiredLocalizationInfoForSaving)
 							.Text(Translation)
 							.HintText(LOCTEXT("TranslationEditTextBox_HintText", "Enter/edit translation here."))
 						]
@@ -382,14 +383,14 @@ void STranslationPickerEditWidget::Construct(const FArguments& InArgs)
 				.Justification(ETextJustify::Center)
 			];
 	}
-	else if (!bHasRequiredLocalizationInfo)
+	else if (!bHasRequiredLocalizationInfoForSaving)
 	{
 		LocalizationInfoAndSaveButtonSlot->AddSlot(0, 0)
 			.Padding(FMargin(5))
 			.ColumnSpan(2)
 			[
 				SNew(STextBlock)
-				.Text(LOCTEXT("RequiredLocalizationInfoNotFound", "The required localization info for this text was not found."))
+				.Text(LOCTEXT("RequiredLocalizationInfoNotFound", "This text is not ready to be localized."))
 				.Justification(ETextJustify::Center)
 			];
 	}
@@ -444,7 +445,7 @@ void STranslationPickerEditWidget::Construct(const FArguments& InArgs)
 				.HAlign(HAlign_Center)
 				.ContentPadding(FEditorStyle::GetMargin("StandardDialog.ContentPadding"))
 				.OnClicked(this, &STranslationPickerEditWidget::SaveAndPreview)
-				.IsEnabled(bHasRequiredLocalizationInfo)
+				.IsEnabled(bHasRequiredLocalizationInfoForSaving)
 				.Visibility(bAllowEditing ? EVisibility::Visible : EVisibility::Collapsed)
 				[
 					SNew(SHorizontalBox)
@@ -454,7 +455,7 @@ void STranslationPickerEditWidget::Construct(const FArguments& InArgs)
 					.AutoWidth()
 					[
 						SNew(STextBlock)
-						.Text(bHasRequiredLocalizationInfo ? LOCTEXT("SaveAndPreviewButtonText", "Save and preview") : LOCTEXT("SaveAndPreviewButtonDisabledText", "Cannot Save"))
+						.Text(bHasRequiredLocalizationInfoForSaving ? LOCTEXT("SaveAndPreviewButtonText", "Save and preview") : LOCTEXT("SaveAndPreviewButtonDisabledText", "Cannot Save"))
 					]
 				]
 			];
diff --git a/Engine/Source/Editor/TranslationEditor/Private/TranslationPickerEditWindow.h b/Engine/Source/Editor/TranslationEditor/Private/TranslationPickerEditWindow.h
--- a/Engine/Source/Editor/TranslationEditor/Private/TranslationPickerEditWindow.h
+++ b/Engine/Source/Editor/TranslationEditor/Private/TranslationPickerEditWindow.h
@@ -80,6 +80,12 @@ class STranslationPickerEditWidget : public SCompoundWidget
 	/** Return the translation unit for this text, with any modifications */
 	UTranslationUnit* GetTranslationUnitWithAnyChanges();
 
+	/** Whether or not we can save */
+	bool CanSave()
+	{
+		return bAllowEditing && bHasRequiredLocalizationInfoForSaving;
+	}
+
 private:
 
 	virtual bool SupportsKeyboardFocus() const override
@@ -100,6 +106,9 @@ private:
 
 	/** Whether or not to show the save button*/
 	bool bAllowEditing;
+
+	/** Whether or not we were able to find the necessary info for saving */
+	bool bHasRequiredLocalizationInfoForSaving;
 };
 
 /** Translation picker edit window to allow you to translate selected FTexts in place */
diff --git a/Engine/Source/Runtime/Core/Private/Internationalization/TextLocalizationManager.cpp b/Engine/Source/Runtime/Core/Private/Internationalization/TextLocalizationManager.cpp
--- a/Engine/Source/Runtime/Core/Private/Internationalization/TextLocalizationManager.cpp
+++ b/Engine/Source/Runtime/Core/Private/Internationalization/TextLocalizationManager.cpp
@@ -420,7 +420,7 @@ bool FTextLocalizationManager::GetLocResID(const FString& Namespace, const FStri
 	const FDisplayStringLookupTable::FDisplayStringEntry* LiveEntry = nullptr;
 	DisplayStringLookupTable.Find(Namespace, LiveKeyTable, Key, LiveEntry);
 
-	if (LiveEntry != nullptr)
+	if (LiveEntry != nullptr && !LiveEntry->LocResID.IsEmpty())
 	{
 		OutLocResId = LiveEntry->LocResID;
 		return true;
