Localization Module: GetLocalizationTargetByName moved here from LocalizationDashboard, removing old version and using new version instead

diff --git a/Engine/Source/Editor/Localization/Public/LocalizationModule.h b/Engine/Source/Editor/Localization/Public/LocalizationModule.h
--- a/Engine/Source/Editor/Localization/Public/LocalizationModule.h
+++ b/Engine/Source/Editor/Localization/Public/LocalizationModule.h
@@ -4,8 +4,13 @@
 
 class ULocalizationTarget;
 
-class FLocalizationModule : public IModuleInterface
+class LOCALIZATION_API FLocalizationModule : public IModuleInterface
 {
 public:
 	ULocalizationTarget* GetLocalizationTargetByName(FString TargetName, bool bIsEngineTarget);
+
+	static FLocalizationModule& Get()
+	{
+		return FModuleManager::LoadModuleChecked<FLocalizationModule>("Localization");
+	}
 };
\ No newline at end of file
diff --git a/Engine/Source/Editor/LocalizationDashboard/Private/LocalizationDashboardModule.cpp b/Engine/Source/Editor/LocalizationDashboard/Private/LocalizationDashboardModule.cpp
--- a/Engine/Source/Editor/LocalizationDashboard/Private/LocalizationDashboardModule.cpp
+++ b/Engine/Source/Editor/LocalizationDashboard/Private/LocalizationDashboardModule.cpp
@@ -153,35 +153,6 @@ public:
 		return LSP ? *LSP : nullptr;
 	}
 
-	ULocalizationTarget* GetLocalizationTargetByName(FString TargetName, bool bIsEngineTarget) override
-	{
-		if (bIsEngineTarget)
-		{
-			ULocalizationTargetSet* EngineTargetSet = ULocalizationSettings::GetEngineTargetSet();
-			for (ULocalizationTarget* Target : EngineTargetSet->TargetObjects)
-			{
-				if (Target->Settings.Name == TargetName)
-				{
-					return Target;
-				}
-			}
-		}
-		else
-		{
-			ULocalizationTargetSet* EngineTargetSet = ULocalizationSettings::GetEngineTargetSet();
-			for (ULocalizationTarget* Target : EngineTargetSet->TargetObjects)
-			{
-				if (Target->Settings.Name == TargetName)
-				{
-					return Target;
-				}
-			}
-		}
-
-		return nullptr;
-	}
-
-
 private:
 	TArray<ILocalizationServiceProvider*> ServiceProviders;
 	FName CurrentServiceProviderName;
diff --git a/Engine/Source/Editor/LocalizationDashboard/Public/ILocalizationDashboardModule.h b/Engine/Source/Editor/LocalizationDashboard/Public/ILocalizationDashboardModule.h
--- a/Engine/Source/Editor/LocalizationDashboard/Public/ILocalizationDashboardModule.h
+++ b/Engine/Source/Editor/LocalizationDashboard/Public/ILocalizationDashboardModule.h
@@ -29,6 +29,4 @@ public:
 	{
 		return FModuleManager::LoadModuleChecked<ILocalizationDashboardModule>("LocalizationDashboard");
 	}
-
-	virtual ULocalizationTarget* GetLocalizationTargetByName(FString TargetName, bool bIsEngineTarget) = 0;
 };
diff --git a/Engine/Source/Editor/TranslationEditor/Private/TranslationDataManager.cpp b/Engine/Source/Editor/TranslationEditor/Private/TranslationDataManager.cpp
--- a/Engine/Source/Editor/TranslationEditor/Private/TranslationDataManager.cpp
+++ b/Engine/Source/Editor/TranslationEditor/Private/TranslationDataManager.cpp
@@ -13,7 +13,7 @@
 #include "NotificationManager.h"
 #include "PortableObjectFormatDOM.h"
 #include "ILocalizationServiceModule.h"
-#include "ILocalizationDashboardModule.h"
+#include "LocalizationModule.h"
 #include "LocalizationTargetTypes.h"
 #include "LocalizationConfigurationScript.h"
 
@@ -929,7 +929,7 @@ bool FTranslationDataManager::SaveSelectedTranslations(TArray<UTranslationUnit*>
 					FString UploadFilePath = FPaths::GameSavedDir() / "Temp" / CultureName / ManifestAndArchiveName + ".po";
 					FFileHelper::SaveStringToFile(PortableObjectDom.ToString(), *UploadFilePath);
 
-					ULocalizationTarget* Target = ILocalizationDashboardModule::Get().GetLocalizationTargetByName(ManifestAndArchiveName, IsEngineManifest);
+					ULocalizationTarget* Target = FLocalizationModule::Get().GetLocalizationTargetByName(ManifestAndArchiveName, IsEngineManifest);
 					FGuid LocalizationTargetGuid;
 					if (Target)
 					{
diff --git a/Engine/Source/Editor/TranslationEditor/TranslationEditor.Build.cs b/Engine/Source/Editor/TranslationEditor/TranslationEditor.Build.cs
--- a/Engine/Source/Editor/TranslationEditor/TranslationEditor.Build.cs
+++ b/Engine/Source/Editor/TranslationEditor/TranslationEditor.Build.cs
@@ -37,7 +37,6 @@ public class TranslationEditor : ModuleRules
                 "Documentation",
                 "Localization",
                 "LocalizationService",
-                "LocalizationDashboard",
 			}
 		);
 
