Translation Picker: Fixed copy & pasted variable name and error message

diff --git a/Engine/Source/Editor/TranslationEditor/Private/TranslationDataManager.cpp b/Engine/Source/Editor/TranslationEditor/Private/TranslationDataManager.cpp
--- a/Engine/Source/Editor/TranslationEditor/Private/TranslationDataManager.cpp
+++ b/Engine/Source/Editor/TranslationEditor/Private/TranslationDataManager.cpp
@@ -968,7 +968,7 @@ bool FTranslationDataManager::SaveSelectedTranslations(TArray<UTranslationUnit*>
 
 void FTranslationDataManager::SaveSelectedTranslationsToTranslationServiceCallback(const FLocalizationServiceOperationRef& Operation, ELocalizationServiceOperationCommandResult::Type Result)
 {
-	TSharedPtr<FUploadLocalizationTargetFile, ESPMode::ThreadSafe> DownloadLocalizationTargetOp = StaticCastSharedRef<FUploadLocalizationTargetFile>(Operation);
+	TSharedPtr<FUploadLocalizationTargetFile, ESPMode::ThreadSafe> UploadLocalizationTargetOp = StaticCastSharedRef<FUploadLocalizationTargetFile>(Operation);
 	bool bError = !(Result == ELocalizationServiceOperationCommandResult::Succeeded);
 	FText ErrorText = FText::GetEmpty();
 	FGuid InTargetGuid;
@@ -977,12 +977,12 @@ void FTranslationDataManager::SaveSelectedTranslationsToTranslationServiceCallba
 	FString TargetName = "";
 	FString TargetPath = "";
 	FString CultureName = "";
-	if (DownloadLocalizationTargetOp.IsValid())
+	if (UploadLocalizationTargetOp.IsValid())
 	{
-		ErrorText = DownloadLocalizationTargetOp->GetOutErrorText();
-		InTargetGuid = DownloadLocalizationTargetOp->GetInTargetGuid();
-		InLocale = DownloadLocalizationTargetOp->GetInLocale();
-		InRelativeOutputFilePathAndName = DownloadLocalizationTargetOp->GetInRelativeInputFilePathAndName();
+		ErrorText = UploadLocalizationTargetOp->GetOutErrorText();
+		InTargetGuid = UploadLocalizationTargetOp->GetInTargetGuid();
+		InLocale = UploadLocalizationTargetOp->GetInLocale();
+		InRelativeOutputFilePathAndName = UploadLocalizationTargetOp->GetInRelativeInputFilePathAndName();
 		TargetName = FPaths::GetBaseFilename(InRelativeOutputFilePathAndName);
 		TargetPath = FPaths::GetPath(InRelativeOutputFilePathAndName);
 		CultureName = FPaths::GetBaseFilename(TargetPath);
@@ -1008,7 +1008,7 @@ void FTranslationDataManager::SaveSelectedTranslationsToTranslationServiceCallba
 	{
 		if (ErrorText.IsEmpty())
 		{
-			ErrorText = LOCTEXT("DownloadLatestFromLocalizationServiceUnspecifiedError", "An unspecified error occured when trying download and import from the Localization Service.");
+			ErrorText = LOCTEXT("SaveToLocalizationServiceUnspecifiedError", "An unspecified error occured when trying to save to the Localization Service.");
 		}
 
 		FText ErrorNotify = FText::Format(LOCTEXT("SaveSelectedTranslationsToTranslationServiceFail", "{0} translations for {1} target failed to save to Translation Service!"), FText::FromString(CultureDisplayName), FText::FromString(TargetName));
