Fixed translation editor support for archive entries with key meta data.

diff --git a/Engine/Source/Editor/TranslationEditor/Private/TranslationDataManager.cpp b/Engine/Source/Editor/TranslationEditor/Private/TranslationDataManager.cpp
--- a/Engine/Source/Editor/TranslationEditor/Private/TranslationDataManager.cpp
+++ b/Engine/Source/Editor/TranslationEditor/Private/TranslationDataManager.cpp
@@ -86,18 +86,13 @@ void FTranslationDataManager::Initialize()
 			int32 NumManifestEntriesParsed = 0;
 
 			GWarn->BeginSlowTask(LOCTEXT("LoadingCurrentManifest", "Loading Entries from Current Translation Manifest..."), true);
-			// Get all manifest entries by source text (same source text in multiple contexts will only show up once)
+
+			// Get all manifest entries by source text (same source text in multiple contexts will only show up once, unless they have unique key metadata)
 			for (auto ManifestItr = ManifestAtHeadRevision->GetEntriesBySourceTextIterator(); ManifestItr; ++ManifestItr, ++NumManifestEntriesParsed)
 			{
 				GWarn->StatusUpdate(NumManifestEntriesParsed, ManifestEntriesCount, FText::Format(LOCTEXT("LoadingCurrentManifestEntries", "Loading Entry {0} of {1} from Current Translation Manifest..."), FText::AsNumber(NumManifestEntriesParsed), FText::AsNumber(ManifestEntriesCount)));
 				const TSharedRef<FManifestEntry> ManifestEntry = ManifestItr.Value();
-				UTranslationUnit* TranslationUnit = NewObject<UTranslationUnit>();
-				check(TranslationUnit != nullptr);
-				// We want Undo/Redo support
-				TranslationUnit->SetFlags(RF_Transactional);
-				TranslationUnit->HasBeenReviewed = false;
-				TranslationUnit->Source = ManifestEntry->Source.Text;
-				TranslationUnit->Namespace = ManifestEntry->Namespace;
+				TMap< TSharedPtr<FLocMetadataObject>, UTranslationUnit* > KeyMetaDataToTranslationUnitMap;
 
 				for(auto ContextIter( ManifestEntry->Contexts.CreateConstIterator() ); ContextIter; ++ContextIter)
 				{
@@ -107,6 +102,20 @@ void FTranslationDataManager::Initialize()
 					ContextInfo.Context = AContext.SourceLocation;
 					ContextInfo.Key = AContext.Key;
 
+					// Make sure we have a unique translation unit for each unique key metadata object.
+					UTranslationUnit*& TranslationUnit = KeyMetaDataToTranslationUnitMap.FindOrAdd(AContext.KeyMetadataObj);
+					if (!TranslationUnit)
+					{
+						TranslationUnit = NewObject<UTranslationUnit>();
+						check(TranslationUnit != nullptr);
+						// We want Undo/Redo support
+						TranslationUnit->SetFlags(RF_Transactional);
+						TranslationUnit->HasBeenReviewed = false;
+						TranslationUnit->Source = ManifestEntry->Source.Text;
+						TranslationUnit->Namespace = ManifestEntry->Namespace;
+						TranslationUnit->KeyMetaDataObject = AContext.KeyMetadataObj;
+					}
+					
 					if (NativeArchivePtr.IsValid() && NativeArchivePtr != ArchivePtr)
 					{
 						const TSharedPtr<FArchiveEntry> NativeArchiveEntry = NativeArchivePtr->FindEntryBySource(ManifestEntry->Namespace, ManifestEntry->Source, AContext.KeyMetadataObj);
@@ -120,7 +129,10 @@ void FTranslationDataManager::Initialize()
 					TranslationUnit->Contexts.Add(ContextInfo);
 				}
 
-				TranslationUnits.Add(TranslationUnit);
+
+				TArray<UTranslationUnit*> TranslationUnitsToAdd;
+				KeyMetaDataToTranslationUnitMap.GenerateValueArray(TranslationUnitsToAdd);
+				TranslationUnits.Append(TranslationUnitsToAdd);
 			}
 			GWarn->EndSlowTask();
 
@@ -673,7 +685,7 @@ void FTranslationDataManager::LoadFromArchive(TArray<UTranslationUnit*>& InTrans
 
 	if (ArchivePtr.IsValid())
 	{
-		TSharedRef< FInternationalizationArchive > Archive = ArchivePtr.ToSharedRef();
+		const TSharedRef< FInternationalizationArchive > Archive = ArchivePtr.ToSharedRef();
 
 		// Make a local copy of this array before we empty the arrays below (we might have been passed AllTranslations array)
 		TArray<UTranslationUnit*> TranslationUnits;
@@ -699,7 +711,7 @@ void FTranslationDataManager::LoadFromArchive(TArray<UTranslationUnit*>& InTrans
 				GWarn->StatusUpdate(CurrentTranslationUnitIndex, TranslationUnits.Num(), FText::Format(LOCTEXT("LoadingCurrentArchiveEntries", "Loading Entry {0} of {1} from Translation Archive..."), FText::AsNumber(CurrentTranslationUnitIndex), FText::AsNumber(TranslationUnits.Num())));
 
 				const FLocItem SourceSearch(TranslationUnit->Source);
-				TSharedPtr<FArchiveEntry> ArchiveEntry = Archive->FindEntryBySource(TranslationUnit->Namespace, SourceSearch, nullptr);
+				TSharedPtr<FArchiveEntry> ArchiveEntry = Archive->FindEntryBySource(TranslationUnit->Namespace, SourceSearch, TranslationUnit->KeyMetaDataObject);
 				if (ArchiveEntry.IsValid())
 				{
 					const FString PreviousTranslation = TranslationUnit->Translation;
diff --git a/Engine/Source/Editor/TranslationEditor/Private/TranslationUnit.h b/Engine/Source/Editor/TranslationEditor/Private/TranslationUnit.h
--- a/Engine/Source/Editor/TranslationEditor/Private/TranslationUnit.h
+++ b/Engine/Source/Editor/TranslationEditor/Private/TranslationUnit.h
@@ -90,6 +90,8 @@ public:
 	DECLARE_EVENT_OneParam(UTranslationUnit, FTranslationUnitPropertyChangedEvent, FName /*PropertyName*/);
 	FTranslationUnitPropertyChangedEvent& OnPropertyChanged() { return TranslationUnitPropertyChangedEvent; }
 
+	TSharedPtr<FLocMetadataObject> KeyMetaDataObject;
+
 protected:
 	/**
 	* Called when a property on this object has been modified externally
