Translation Picker: Fix small bug with Preview in Editor

diff --git a/Engine/Source/Editor/TranslationEditor/Private/TranslationDataManager.cpp b/Engine/Source/Editor/TranslationEditor/Private/TranslationDataManager.cpp
--- a/Engine/Source/Editor/TranslationEditor/Private/TranslationDataManager.cpp
+++ b/Engine/Source/Editor/TranslationEditor/Private/TranslationDataManager.cpp
@@ -27,15 +27,13 @@ FTranslationDataManager::FTranslationDataManager( const FString& InManifestFileP
 	: OpenedManifestFilePath(InManifestFilePath)
 	, NativeArchiveFilePath(InNativeArchiveFilePath)
 	, OpenedArchiveFilePath(InArchiveFilePath)
-	, AssociatedLocalizationTarget(nullptr)
 	, bLoadedSuccessfully(true)
 {
 	Initialize();
 }
 
 FTranslationDataManager::FTranslationDataManager(ULocalizationTarget* const LocalizationTarget, const FString& CultureToEdit)
-	: AssociatedLocalizationTarget(LocalizationTarget)
-	, bLoadedSuccessfully(true)
+	: bLoadedSuccessfully(true)
 {
 	check(LocalizationTarget);
  
@@ -612,7 +610,7 @@ void FTranslationDataManager::HandlePropertyChanged(FName PropertyName)
 	WriteTranslationData();
 }
 
-void FTranslationDataManager::PreviewAllTranslationsInEditor()
+void FTranslationDataManager::PreviewAllTranslationsInEditor(ULocalizationTarget* LocalizationTarget)
 {
 	FString ManifestFullPath = FPaths::ConvertRelativePathToFull(OpenedManifestFilePath);
 	FString EngineFullPath = FPaths::ConvertRelativePathToFull(FPaths::EngineContentDir());
@@ -623,10 +621,10 @@ void FTranslationDataManager::PreviewAllTranslationsInEditor()
 		IsEngineManifest = true;
 	}
 
-	if (AssociatedLocalizationTarget.IsValid())
+	if (LocalizationTarget != nullptr)
 	{
-		const FString ConfigFilePath = LocalizationConfigurationScript::GetRegenerateResourcesScriptPath(AssociatedLocalizationTarget.Get());
-		LocalizationConfigurationScript::GenerateRegenerateResourcesScript(AssociatedLocalizationTarget.Get()).Write(ConfigFilePath);
+		const FString ConfigFilePath = LocalizationConfigurationScript::GetRegenerateResourcesScriptPath(LocalizationTarget);
+		LocalizationConfigurationScript::GenerateRegenerateResourcesScript(LocalizationTarget).Write(ConfigFilePath);
 
 		FJsonInternationalizationArchiveSerializer LocalizationArchiveSerializer;
 		FJsonInternationalizationManifestSerializer LocalizationManifestSerializer;
@@ -887,6 +885,8 @@ bool FTranslationDataManager::SaveSelectedTranslations(TArray<UTranslationUnit*>
 			IsEngineManifest = true;
 		}
 
+		ULocalizationTarget* LocalizationTarget = FLocalizationModule::Get().GetLocalizationTargetByName(ManifestAndArchiveName, IsEngineManifest);;
+
 		if (FPaths::FileExists(ManifestFullPath) && FPaths::FileExists(ArchiveFullPath))
 		{
 			TSharedRef<FTranslationDataManager> DataManager = MakeShareable(new FTranslationDataManager(ManifestFullPath, NativeArchiveFullPath, ArchiveFullPath));
@@ -941,11 +941,10 @@ bool FTranslationDataManager::SaveSelectedTranslations(TArray<UTranslationUnit*>
 					FString UploadFilePath = FPaths::GameSavedDir() / "Temp" / CultureName / ManifestAndArchiveName + ".po";
 					FFileHelper::SaveStringToFile(PortableObjectDom.ToString(), *UploadFilePath);
 
-					ULocalizationTarget* Target = FLocalizationModule::Get().GetLocalizationTargetByName(ManifestAndArchiveName, IsEngineManifest);
 					FGuid LocalizationTargetGuid;
-					if (Target)
+					if (LocalizationTarget)
 					{
-						LocalizationTargetGuid = Target->Settings.Guid;
+						LocalizationTargetGuid = LocalizationTarget->Settings.Guid;
 					}
 
 					ILocalizationServiceProvider& Provider = ILocalizationServiceModule::Get().GetProvider();
@@ -967,7 +966,7 @@ bool FTranslationDataManager::SaveSelectedTranslations(TArray<UTranslationUnit*>
 
 			// Save the data to file, and preview in editor
 			bSucceeded = bSucceeded && DataManager->WriteTranslationData();
-			DataManager->PreviewAllTranslationsInEditor();
+			DataManager->PreviewAllTranslationsInEditor(LocalizationTarget);
 		}
 		else
 		{
diff --git a/Engine/Source/Editor/TranslationEditor/Private/TranslationDataManager.h b/Engine/Source/Editor/TranslationEditor/Private/TranslationDataManager.h
--- a/Engine/Source/Editor/TranslationEditor/Private/TranslationDataManager.h
+++ b/Engine/Source/Editor/TranslationEditor/Private/TranslationDataManager.h
@@ -61,7 +61,7 @@ public:
 	void HandlePropertyChanged(FName PropertyName);
 
 	/** Regenerate and reload archives to reflect modifications in the UI */
-	void PreviewAllTranslationsInEditor();
+	void PreviewAllTranslationsInEditor(ULocalizationTarget* LocalizationTarget);
 
 	/** Put items in the Search Array if they match this filter */
 	void PopulateSearchResultsUsingFilter(const FString& SearchFilter);
diff --git a/Engine/Source/Editor/TranslationEditor/Private/TranslationEditor.cpp b/Engine/Source/Editor/TranslationEditor/Private/TranslationEditor.cpp
--- a/Engine/Source/Editor/TranslationEditor/Private/TranslationEditor.cpp
+++ b/Engine/Source/Editor/TranslationEditor/Private/TranslationEditor.cpp
@@ -1135,7 +1135,7 @@ void FTranslationEditor::UpdateContextSelection()
 
 void FTranslationEditor::PreviewAllTranslationsInEditor_Execute()
 {
-	DataManager->PreviewAllTranslationsInEditor();
+	DataManager->PreviewAllTranslationsInEditor(AssociatedLocalizationTarget.Get());
 }
 
 void FTranslationEditor::ImportLatestFromLocalizationService_Execute()
