Translation Data Manager: Fixed small bug where NativeCultureName was not initialized, preventing saving of selected translations via the Translation Picker.

diff --git a/Engine/Source/Editor/TranslationEditor/Private/TranslationDataManager.cpp b/Engine/Source/Editor/TranslationEditor/Private/TranslationDataManager.cpp
--- a/Engine/Source/Editor/TranslationEditor/Private/TranslationDataManager.cpp
+++ b/Engine/Source/Editor/TranslationEditor/Private/TranslationDataManager.cpp
@@ -868,14 +868,10 @@ bool FTranslationDataManager::SaveSelectedTranslations(TArray<UTranslationUnit*>
 		auto Item = *TextIt;
 		FString CurrentLocResPath = Item.Key;
 		FString ManifestAndArchiveName = FPaths::GetBaseFilename(CurrentLocResPath);
-		
+
 		FString ArchiveFilePath = FPaths::GetPath(CurrentLocResPath);
 		FString CultureName = FPaths::GetBaseFilename(ArchiveFilePath);
 		FString ManifestPath = FPaths::GetPath(ArchiveFilePath);
-
-		FString NativeCultureName;
-
-		FString NativeArchiveFullPath = ManifestPath / NativeCultureName / ManifestAndArchiveName + ".archive";
 		FString ArchiveFullPath = ArchiveFilePath / ManifestAndArchiveName + ".archive";
 		FString ManifestFullPath = ManifestPath / ManifestAndArchiveName + ".manifest";
 		FString EngineFullPath = FPaths::ConvertRelativePathToFull(FPaths::EngineContentDir());
@@ -885,7 +881,14 @@ bool FTranslationDataManager::SaveSelectedTranslations(TArray<UTranslationUnit*>
 			IsEngineManifest = true;
 		}
 
-		ULocalizationTarget* LocalizationTarget = FLocalizationModule::Get().GetLocalizationTargetByName(ManifestAndArchiveName, IsEngineManifest);;
+		ULocalizationTarget* LocalizationTarget = FLocalizationModule::Get().GetLocalizationTargetByName(ManifestAndArchiveName, IsEngineManifest);
+
+		FString NativeCultureName;
+		if (LocalizationTarget->Settings.SupportedCulturesStatistics.IsValidIndex(LocalizationTarget->Settings.NativeCultureIndex))
+		{
+			NativeCultureName = LocalizationTarget->Settings.SupportedCulturesStatistics[LocalizationTarget->Settings.NativeCultureIndex].CultureName;
+		}
+		FString NativeArchiveFullPath = ManifestPath / NativeCultureName / ManifestAndArchiveName + ".archive";
 
 		if (FPaths::FileExists(ManifestFullPath) && FPaths::FileExists(ArchiveFullPath))
 		{
