Copying //UE4/Dev-Editor to //UE4/Dev-Main (Source: //UE4/Dev-Editor @ 3082391)

diff --git a/Engine/Source/Editor/PropertyEditor/Public/IDetailsView.h b/Engine/Source/Editor/PropertyEditor/Public/IDetailsView.h
--- a/Engine/Source/Editor/PropertyEditor/Public/IDetailsView.h
+++ b/Engine/Source/Editor/PropertyEditor/Public/IDetailsView.h
@@ -31,6 +31,11 @@ struct FDetailsViewArgs
 		Automatic,
 	};
 
+	/** Controls how CPF_DisableEditOnInstance nodes will be treated */
+	EEditDefaultsOnlyNodeVisibility DefaultsOnlyVisibility;
+	/** The command list from the host of the details view, allowing child widgets to bind actions with a bound chord */
+	TSharedPtr<class FUICommandList> HostCommandList;
+
 	/** Identifier for this details view; NAME_None if this view is anonymous */
 	FName ViewIdentifier;
 	/** Notify hook to call when properties are changed */
@@ -61,10 +66,9 @@ struct FDetailsViewArgs
 	uint32 bCustomFilterAreaLocation : 1;
 	/** If false, the current properties editor will never display the favorite system */
 	uint32 bAllowFavoriteSystem : 1;
-	/** Controls how CPF_DisableEditOnInstance nodes will be treated */
-	EEditDefaultsOnlyNodeVisibility DefaultsOnlyVisibility;
-	/** The command list from the host of the details view, allowing child widgets to bind actions with a bound chord */
-	TSharedPtr<class FUICommandList> HostCommandList;
+	/** If true the details panel will assume each object passed in through SetObjects will be a unique object shown in the tree and not combined with other objects */
+	uint32 bAllowMultipleTopLevelObjects : 1;
+
 
 public:
 	/** Default constructor */
@@ -76,7 +80,8 @@ public:
 					, FNotifyHook* InNotifyHook = NULL
 					, const bool InSearchInitialKeyFocus = false
 					, FName InViewIdentifier = NAME_None )
-		: ViewIdentifier( InViewIdentifier )
+		: DefaultsOnlyVisibility(EEditDefaultsOnlyNodeVisibility::Show)
+		, ViewIdentifier( InViewIdentifier )
 		, NotifyHook( InNotifyHook )
 		, NameAreaSettings( InNameAreaSettings )
 		, bUpdatesFromSelection( InUpdateFromSelection )
@@ -91,7 +96,7 @@ public:
 		, bCustomNameAreaLocation(false)
 		, bCustomFilterAreaLocation(false)
 		, bAllowFavoriteSystem(true)
-		, DefaultsOnlyVisibility(EEditDefaultsOnlyNodeVisibility::Show)
+		, bAllowMultipleTopLevelObjects(false)
 	{
 	}
 };
@@ -104,7 +109,7 @@ class IDetailsView : public SCompoundWidget
 {
 public:
 	/** Sets the callback for when the property view changes */
-	virtual void SetOnObjectArrayChanged( FOnObjectArrayChanged OnObjectArrayChangedDelegate) = 0;
+	virtual void SetOnObjectArrayChanged(FOnObjectArrayChanged OnObjectArrayChangedDelegate) = 0;
 
 	/** List of all selected objects we are inspecting */
 	virtual const TArray< TWeakObjectPtr<UObject> >& GetSelectedObjects() const = 0;
@@ -119,10 +124,11 @@ public:
 	virtual bool HasClassDefaultObject() const = 0;
 
 	/** Gets the base class being viewed */
+	DEPRECATED(4.13, "GetBaseClass should be accessed from IDetailLayoutBuilder.GetBaseClass")
 	virtual const UClass* GetBaseClass() const = 0;
-	
-	/** Gets the base struct being viewed */
-	virtual UStruct* GetBaseStruct() const = 0;
+
+	DEPRECATED(4.13, "GetBaseClass should be accessed from IDetailLayoutBuilder.GetBaseClass")
+	virtual UClass* GetBaseClass() = 0;
 
 	/**
 	 * Registers a custom detail layout delegate for a specific class in this instance of the details view only
@@ -132,6 +138,15 @@ public:
 	 */
 	virtual void RegisterInstancedCustomPropertyLayout( UStruct* Class, FOnGetDetailCustomizationInstance DetailLayoutDelegate ) = 0;
 
+
+	/**
+	* Registers a customization that will be used only if this details panel contains multiple top level objects.
+	* I.E it was created with bAllowMultipleTopLevelObjects = true.	This interface will be used to customize the header for each top level object in the details panel
+	*
+	* @param InRootObjectCustomization	If null is passed in, the customization will be removed
+	*/
+	virtual void SetRootObjectCustomizationInstance( TSharedPtr<class IDetailRootObjectCustomization> InRootObjectCustomization ) = 0;
+
 	/**
 	 * Unregisters a custom detail layout delegate for a specific class in this instance of the details view only
 	 *
@@ -169,10 +184,26 @@ public:
 	virtual bool IsLocked() const = 0;
 
 	/**
-	 * @return true of the details view can be updated from editor selection
+	 * @return true if the details view can be updated from editor selection
 	 */
 	virtual bool IsUpdatable() const = 0;
 
+	/**
+	 * @return True if there is any filter of properties active in this details panel
+	 */
+	virtual bool HasActiveSearch() const = 0;
+	
+	/** 
+	 * Clears any search terms in the current filter
+	 */
+	virtual void ClearSearch() = 0;
+
+	/**
+	 * @return The number of visible top level objects. This value is affected by filtering. 
+	 * Note: this value will always be 1 if this details panel was not created with bAllowMultipleTopLevelObjects=true
+	 */
+	virtual int32 GetNumVisibleTopLevelObjects() const = 0;
+
 	/** @return The identifier for this details view, or NAME_None is this view is anonymous */
 	virtual FName GetIdentifier() const = 0;
 
@@ -234,6 +265,11 @@ public:
 	 */
 	virtual void ShowAllAdvancedProperties() = 0;
 	
+	/**
+	 * Refreshes the visibility of root objects in this details view. 
+	 * Note: This method has no effect if the details panel is viewing a single top-level object set only
+	 */
+	virtual void RefreshRootObjectVisibility() = 0;
 	/**
 	 * Assigns delegate called when view is filtered, useful for updating external control logic:
 	 */
diff --git a/Engine/Source/Editor/PropertyEditor/Public/PropertyEditing.h b/Engine/Source/Editor/PropertyEditor/Public/PropertyEditing.h
--- a/Engine/Source/Editor/PropertyEditor/Public/PropertyEditing.h
+++ b/Engine/Source/Editor/PropertyEditor/Public/PropertyEditing.h
@@ -14,3 +14,4 @@
 #include "IDetailChildrenBuilder.h"
 #include "IDetailCustomization.h"
 #include "IDetailsView.h"
+#include "IDetailRootObjectCustomization.h"
diff --git a/Engine/Source/Editor/ReferenceViewer/Classes/EdGraph_ReferenceViewer.h b/Engine/Source/Editor/ReferenceViewer/Classes/EdGraph_ReferenceViewer.h
--- a/Engine/Source/Editor/ReferenceViewer/Classes/EdGraph_ReferenceViewer.h
+++ b/Engine/Source/Editor/ReferenceViewer/Classes/EdGraph_ReferenceViewer.h
@@ -20,9 +20,17 @@ public:
 
 	bool IsSearchDepthLimited() const;
 	bool IsSearchBreadthLimited() const;
+	bool IsShowSoftReferences() const;
+	bool IsShowSoftDependencies() const;
+	bool IsShowHardReferences() const;
+	bool IsShowHardDependencies() const;
 
 	void SetSearchDepthLimitEnabled(bool newEnabled);
 	void SetSearchBreadthLimitEnabled(bool newEnabled);
+	void SetShowSoftReferencesEnabled(bool newEnabled);
+	void SetShowSoftDependenciesEnabled(bool newEnabled);
+	void SetShowHardReferencesEnabled(bool newEnabled);
+	void SetShowHardDependenciesEnabled(bool newEnabled);
 
 	int32 GetSearchDepthLimit() const;
 	int32 GetSearchBreadthLimit() const;
@@ -59,5 +67,9 @@ private:
 
 	bool bLimitSearchDepth;
 	bool bLimitSearchBreadth;
+	bool bIsShowSoftReferences;
+	bool bIsShowSoftDependencies;
+	bool bIsShowHardReferences;
+	bool bIsShowHardDependencies;
 };
 
diff --git a/Engine/Source/Editor/ReferenceViewer/Private/EdGraph_ReferenceViewer.cpp b/Engine/Source/Editor/ReferenceViewer/Private/EdGraph_ReferenceViewer.cpp
--- a/Engine/Source/Editor/ReferenceViewer/Private/EdGraph_ReferenceViewer.cpp
+++ b/Engine/Source/Editor/ReferenceViewer/Private/EdGraph_ReferenceViewer.cpp
@@ -14,6 +14,10 @@ UEdGraph_ReferenceViewer::UEdGraph_ReferenceViewer(const FObjectInitializer& Obj
 
 	bLimitSearchDepth = true;
 	bLimitSearchBreadth = true;
+	bIsShowSoftReferences = true;
+	bIsShowSoftDependencies = true;
+	bIsShowHardReferences = true;
+	bIsShowHardDependencies = true;
 }
 
 void UEdGraph_ReferenceViewer::BeginDestroy()
@@ -57,6 +61,26 @@ bool UEdGraph_ReferenceViewer::IsSearchBreadthLimited() const
 	return bLimitSearchBreadth;
 }
 
+bool UEdGraph_ReferenceViewer::IsShowSoftReferences() const
+{
+	return bIsShowSoftReferences;
+}
+
+bool UEdGraph_ReferenceViewer::IsShowSoftDependencies() const
+{
+	return bIsShowSoftDependencies;
+}
+
+bool UEdGraph_ReferenceViewer::IsShowHardReferences() const
+{
+	return bIsShowHardReferences;
+}
+
+bool UEdGraph_ReferenceViewer::IsShowHardDependencies() const
+{
+	return bIsShowHardDependencies;
+}
+
 void UEdGraph_ReferenceViewer::SetSearchDepthLimitEnabled(bool newEnabled)
 {
 	bLimitSearchDepth = newEnabled;
@@ -67,6 +91,26 @@ void UEdGraph_ReferenceViewer::SetSearchBreadthLimitEnabled(bool newEnabled)
 	bLimitSearchBreadth = newEnabled;
 }
 
+void UEdGraph_ReferenceViewer::SetShowSoftReferencesEnabled(bool newEnabled)
+{
+	bIsShowSoftReferences = newEnabled;
+}
+
+void UEdGraph_ReferenceViewer::SetShowSoftDependenciesEnabled(bool newEnabled)
+{
+	bIsShowSoftDependencies = newEnabled;
+}
+
+void UEdGraph_ReferenceViewer::SetShowHardReferencesEnabled(bool newEnabled)
+{
+	bIsShowHardReferences = newEnabled;
+}
+
+void UEdGraph_ReferenceViewer::SetShowHardDependenciesEnabled(bool newEnabled)
+{
+	bIsShowHardDependencies = newEnabled;
+}
+
 int32 UEdGraph_ReferenceViewer::GetSearchDepthLimit() const
 {
 	return MaxSearchDepth;
@@ -232,18 +276,43 @@ UEdGraphNode_Reference* UEdGraph_ReferenceViewer::RecursivelyConstructNodes(bool
 
 	FAssetRegistryModule& AssetRegistryModule = FModuleManager::LoadModuleChecked<FAssetRegistryModule>("AssetRegistry");
 	TArray<FName> ReferenceNames;
+	TArray<FName> HardReferenceNames;
 	if ( bReferencers )
 	{
 		for ( auto NameIt = PackageNames.CreateConstIterator(); NameIt; ++NameIt )
 		{
-			AssetRegistryModule.Get().GetReferencers(*NameIt, ReferenceNames);
+			AssetRegistryModule.Get().GetReferencers(*NameIt, HardReferenceNames, EAssetRegistryDependencyType::Hard);
+			if (bIsShowSoftReferences && bIsShowHardReferences)
+			{
+				AssetRegistryModule.Get().GetReferencers(*NameIt, ReferenceNames, EAssetRegistryDependencyType::All);
+			}
+			else if (bIsShowSoftReferences && !bIsShowHardReferences)
+			{
+				AssetRegistryModule.Get().GetReferencers(*NameIt, ReferenceNames, EAssetRegistryDependencyType::Soft);
+			}
+			else
+			{
+				AssetRegistryModule.Get().GetReferencers(*NameIt, ReferenceNames, EAssetRegistryDependencyType::Hard);
+			}
 		}
 	}
 	else
 	{
 		for ( auto NameIt = PackageNames.CreateConstIterator(); NameIt; ++NameIt )
 		{
-			AssetRegistryModule.Get().GetDependencies(*NameIt, ReferenceNames);
+			AssetRegistryModule.Get().GetDependencies(*NameIt, HardReferenceNames, EAssetRegistryDependencyType::Hard);
+			if (bIsShowSoftDependencies && bIsShowHardDependencies)
+			{
+				AssetRegistryModule.Get().GetDependencies(*NameIt, ReferenceNames, EAssetRegistryDependencyType::All);
+			}
+			else if (bIsShowSoftDependencies && !bIsShowHardDependencies)
+			{
+				AssetRegistryModule.Get().GetDependencies(*NameIt, ReferenceNames, EAssetRegistryDependencyType::Soft);
+			}
+			else
+			{
+				AssetRegistryModule.Get().GetDependencies(*NameIt, ReferenceNames, EAssetRegistryDependencyType::Hard);
+			}
 		}
 	}
 
@@ -285,6 +354,21 @@ UEdGraphNode_Reference* UEdGraph_ReferenceViewer::RecursivelyConstructNodes(bool
 					TArray<FName> NewPackageNames;
 					NewPackageNames.Add(ReferenceName);
 					UEdGraphNode_Reference* ReferenceNode = RecursivelyConstructNodes(bReferencers, RootNode, NewPackageNames, RefNodeLoc, NodeSizes, PackagesToAssetDataMap, CurrentDepth + 1, VisitedNames);
+					for (FName NodeName : HardReferenceNames)
+					{
+						FName NewNodeFName = FName(*ReferenceNode->NodeComment);
+						if (NodeName == FName(*ReferenceNode->NodeComment))
+						{
+							if (bReferencers)
+							{
+								ReferenceNode->GetDependencyPin()->PinType.PinCategory = TEXT("hard");
+							}
+							else
+							{
+								ReferenceNode->GetReferencerPin()->PinType.PinCategory = TEXT("hard");
+							}
+						}
+					}
 					if ( ensure(ReferenceNode) )
 					{
 						if ( bReferencers )
diff --git a/Engine/Source/Editor/ReferenceViewer/Private/ReferenceViewerSchema.cpp b/Engine/Source/Editor/ReferenceViewer/Private/ReferenceViewerSchema.cpp
--- a/Engine/Source/Editor/ReferenceViewer/Private/ReferenceViewerSchema.cpp
+++ b/Engine/Source/Editor/ReferenceViewer/Private/ReferenceViewerSchema.cpp
@@ -7,6 +7,10 @@
 
 #include "Editor/GraphEditor/Public/ConnectionDrawingPolicy.h"
 
+
+static const FLinearColor RiceFlower = FLinearColor(FColor(236, 252, 227));
+static const FLinearColor CannonPink = FLinearColor(FColor(145, 66, 117));
+
 // Overridden connection drawing policy to use less curvy lines between nodes
 class FReferenceViewerConnectionDrawingPolicy : public FConnectionDrawingPolicy
 {
@@ -21,6 +25,18 @@ public:
 		const int32 Tension = FMath::Abs<int32>(Start.X - End.X);
 		return Tension * FVector2D(1.0f, 0);
 	}
+
+	virtual void DetermineWiringStyle(UEdGraphPin* OutputPin, UEdGraphPin* InputPin, /*inout*/ FConnectionParams& Params) override
+	{
+		if (OutputPin->PinType.PinCategory == TEXT("hard") || InputPin->PinType.PinCategory == TEXT("hard"))
+		{
+			Params.WireColor = RiceFlower;
+		}
+		else
+		{
+			Params.WireColor = CannonPink;
+		}
+	}
 };
 
 //////////////////////////////////////////////////////////////////////////
@@ -50,7 +66,14 @@ void UReferenceViewerSchema::GetContextMenuActions(const UEdGraph* CurrentGraph,
 
 FLinearColor UReferenceViewerSchema::GetPinTypeColor(const FEdGraphPinType& PinType) const
 {
-	return FLinearColor::White;
+	if (PinType.PinCategory == TEXT("hard"))
+	{
+		return RiceFlower;
+	}
+	else
+	{
+		return CannonPink;
+	}
 }
 
 void UReferenceViewerSchema::BreakPinLinks(UEdGraphPin& TargetPin, bool bSendsNodeNotifcation) const
diff --git a/Engine/Source/Editor/ReferenceViewer/Private/SReferenceViewer.cpp b/Engine/Source/Editor/ReferenceViewer/Private/SReferenceViewer.cpp
--- a/Engine/Source/Editor/ReferenceViewer/Private/SReferenceViewer.cpp
+++ b/Engine/Source/Editor/ReferenceViewer/Private/SReferenceViewer.cpp
@@ -224,6 +224,102 @@ void SReferenceViewer::Construct( const FArguments& InArgs )
 							]
 						]
 					]
+
+					+SVerticalBox::Slot()
+					.AutoHeight()
+					[
+						SNew(SHorizontalBox)
+
+						+SHorizontalBox::Slot()
+						.VAlign(VAlign_Center)
+						.Padding(2.f)
+						[
+							SNew(STextBlock)
+							.Text(LOCTEXT("ShowHideSoftReferences", "Show Soft References"))
+						]
+
+						+SHorizontalBox::Slot()
+						.AutoWidth()
+						.VAlign(VAlign_Center)
+						.Padding(2.f)
+						[
+							SNew(SCheckBox)
+							.OnCheckStateChanged( this, &SReferenceViewer::OnShowSoftReferencesChanged )
+							.IsChecked( this, &SReferenceViewer::IsShowSoftReferencesChecked )
+						]
+					]
+
+					+ SVerticalBox::Slot()
+					.AutoHeight()
+					[
+						SNew(SHorizontalBox)
+
+						+ SHorizontalBox::Slot()
+						.VAlign(VAlign_Center)
+						.Padding(2.f)
+						[
+							SNew(STextBlock)
+							.Text(LOCTEXT("ShowHideSoftDependencies", "Show Soft Dependencies"))
+						]
+
+						+SHorizontalBox::Slot()
+						.AutoWidth()
+						.VAlign(VAlign_Center)
+						.Padding(2.f)
+						[
+							SNew(SCheckBox)
+							.OnCheckStateChanged(this, &SReferenceViewer::OnShowSoftDependenciesChanged)
+							.IsChecked(this, &SReferenceViewer::IsShowSoftDependenciesChecked)
+						]
+					]
+
+					+ SVerticalBox::Slot()
+					.AutoHeight()
+					[
+						SNew(SHorizontalBox)
+
+						+ SHorizontalBox::Slot()
+						.VAlign(VAlign_Center)
+						.Padding(2.f)
+						[
+							SNew(STextBlock)
+							.Text(LOCTEXT("ShowHideHardReferences", "Show Hard References"))
+						]
+
+						+SHorizontalBox::Slot()
+						.AutoWidth()
+						.VAlign(VAlign_Center)
+						.Padding(2.f)
+						[
+							SNew(SCheckBox)
+							.OnCheckStateChanged(this, &SReferenceViewer::OnShowHardReferencesChanged)
+							.IsChecked(this, &SReferenceViewer::IsShowHardReferencesChecked)
+						]
+					]
+
+					+ SVerticalBox::Slot()
+					.AutoHeight()
+					[
+						SNew(SHorizontalBox)
+
+						+ SHorizontalBox::Slot()
+						.VAlign(VAlign_Center)
+						.Padding(2.f)
+						[
+							SNew(STextBlock)
+							.Text(LOCTEXT("ShowHideHardDependencies", "Show Hard Dependencies"))
+						]
+
+						+SHorizontalBox::Slot()
+						.AutoWidth()
+						.VAlign(VAlign_Center)
+						.Padding(2.f)
+						[
+							SNew(SCheckBox)
+							.OnCheckStateChanged(this, &SReferenceViewer::OnShowHardDependenciesChanged)
+							.IsChecked(this, &SReferenceViewer::IsShowHardDependenciesChecked)
+						]
+					]
 				]
 			]
 
@@ -475,6 +571,91 @@ ECheckBoxState SReferenceViewer::IsSearchBreadthEnabledChecked() const
 	}
 }
 
+void SReferenceViewer::OnShowSoftReferencesChanged(ECheckBoxState NewState)
+{
+	if (GraphObj)
+	{
+		GraphObj->SetShowSoftReferencesEnabled(NewState == ECheckBoxState::Checked);
+		GraphObj->RebuildGraph();
+	}
+}
+
+ECheckBoxState SReferenceViewer::IsShowSoftReferencesChecked() const
+{
+	if (GraphObj)
+	{
+		return GraphObj->IsShowSoftReferences()? ECheckBoxState::Checked : ECheckBoxState::Unchecked;
+	}
+	else
+	{
+		return ECheckBoxState::Unchecked;
+	}
+}
+
+void SReferenceViewer::OnShowSoftDependenciesChanged(ECheckBoxState NewState)
+{
+	if (GraphObj)
+	{
+		GraphObj->SetShowSoftDependenciesEnabled(NewState == ECheckBoxState::Checked);
+		GraphObj->RebuildGraph();
+	}
+}
+
+ECheckBoxState SReferenceViewer::IsShowSoftDependenciesChecked() const
+{
+	if (GraphObj)
+	{
+		return GraphObj->IsShowSoftDependencies() ? ECheckBoxState::Checked : ECheckBoxState::Unchecked;
+	}
+	else
+	{
+		return ECheckBoxState::Unchecked;
+	}
+}
+
+void SReferenceViewer::OnShowHardReferencesChanged(ECheckBoxState NewState)
+{
+	if (GraphObj)
+	{
+		GraphObj->SetShowHardReferencesEnabled(NewState == ECheckBoxState::Checked);
+		GraphObj->RebuildGraph();
+	}
+}
+
+
+ECheckBoxState SReferenceViewer::IsShowHardReferencesChecked() const
+{
+	if (GraphObj)
+	{
+		return GraphObj->IsShowHardReferences() ? ECheckBoxState::Checked : ECheckBoxState::Unchecked;
+	}
+	else
+	{
+		return ECheckBoxState::Unchecked;
+	}
+}
+
+void SReferenceViewer::OnShowHardDependenciesChanged(ECheckBoxState NewState)
+{
+	if (GraphObj)
+	{
+		GraphObj->SetShowHardDependenciesEnabled(NewState == ECheckBoxState::Checked);
+		GraphObj->RebuildGraph();
+	}
+}
+
+ECheckBoxState SReferenceViewer::IsShowHardDependenciesChecked() const
+{
+	if (GraphObj)
+	{
+		return GraphObj->IsShowHardDependencies() ? ECheckBoxState::Checked : ECheckBoxState::Unchecked;
+	}
+	else
+	{
+		return ECheckBoxState::Unchecked;
+	}
+}
+
 int32 SReferenceViewer::GetSearchBreadthCount() const
 {
 	if ( GraphObj )
diff --git a/Engine/Source/Editor/ReferenceViewer/Private/SReferenceViewer.h b/Engine/Source/Editor/ReferenceViewer/Private/SReferenceViewer.h
--- a/Engine/Source/Editor/ReferenceViewer/Private/SReferenceViewer.h
+++ b/Engine/Source/Editor/ReferenceViewer/Private/SReferenceViewer.h
@@ -74,6 +74,17 @@ private:
 
 	void OnSearchBreadthEnabledChanged( ECheckBoxState NewState );
 	ECheckBoxState IsSearchBreadthEnabledChecked() const;
+
+	void OnShowSoftReferencesChanged( ECheckBoxState NewState );
+	ECheckBoxState IsShowSoftReferencesChecked() const;
+	void OnShowSoftDependenciesChanged(ECheckBoxState NewState);
+	ECheckBoxState IsShowSoftDependenciesChecked() const;
+
+	void OnShowHardReferencesChanged(ECheckBoxState NewState);
+	ECheckBoxState IsShowHardReferencesChecked() const;
+	void OnShowHardDependenciesChanged(ECheckBoxState NewState);
+	ECheckBoxState IsShowHardDependenciesChecked() const;
+
 	int32 GetSearchBreadthCount() const;
 	void OnSearchBreadthCommitted(int32 NewValue);
 
diff --git a/Engine/Source/Editor/SceneOutliner/Private/SSceneOutliner.cpp b/Engine/Source/Editor/SceneOutliner/Private/SSceneOutliner.cpp
--- a/Engine/Source/Editor/SceneOutliner/Private/SSceneOutliner.cpp
+++ b/Engine/Source/Editor/SceneOutliner/Private/SSceneOutliner.cpp
@@ -330,6 +330,9 @@ namespace SceneOutliner
 
 				// Header for the tree
 				.HeaderRow( HeaderRowWidget )
+
+				// Called when an item is expanded or collapsed with the shift-key pressed down
+				.OnSetExpansionRecursive(this, &SSceneOutliner::SetItemExpansionRecursive)
 			]
 		];
 
@@ -2573,6 +2576,21 @@ namespace SceneOutliner
 	{
 		return (InWorld == SharedData->UserChosenWorld);
 	}
+
+	void SSceneOutliner::SetItemExpansionRecursive(FTreeItemPtr Model, bool bInExpansionState)
+	{
+		if (Model.IsValid())
+		{
+			OutlinerTreeView->SetItemExpansion(Model, bInExpansionState);
+			for (auto& Child : Model->Children)
+			{
+				if (Child.IsValid())
+				{
+					SetItemExpansionRecursive(Child.Pin(), bInExpansionState);
+				}
+			}
+		}
+	}
 }	// namespace SceneOutliner
 
 #undef LOCTEXT_NAMESPACE
diff --git a/Engine/Source/Editor/SceneOutliner/Private/SSceneOutliner.h b/Engine/Source/Editor/SceneOutliner/Private/SSceneOutliner.h
--- a/Engine/Source/Editor/SceneOutliner/Private/SSceneOutliner.h
+++ b/Engine/Source/Editor/SceneOutliner/Private/SSceneOutliner.h
@@ -518,6 +518,9 @@ namespace SceneOutliner
 
 		/** Display a checkbox next to the world we are viewing */
 		bool IsWorldChecked(TWeakObjectPtr<UWorld> InWorld);
+
+		/** Handler for recursively expanding/collapsing items */
+		void SetItemExpansionRecursive(FTreeItemPtr Model, bool bInExpansionState);
 	};
 
 }		// namespace SceneOutliner
diff --git a/Engine/Source/Editor/Sequencer/Private/DisplayNodes/SequencerDisplayNode.cpp b/Engine/Source/Editor/Sequencer/Private/DisplayNodes/SequencerDisplayNode.cpp
--- a/Engine/Source/Editor/Sequencer/Private/DisplayNodes/SequencerDisplayNode.cpp
+++ b/Engine/Source/Editor/Sequencer/Private/DisplayNodes/SequencerDisplayNode.cpp
@@ -15,6 +15,7 @@
 #include "IKeyArea.h"
 #include "GroupedKeyArea.h"
 #include "ObjectEditorUtils.h"
+#include "GenericCommands.h"
 
 
 #define LOCTEXT_NAMESPACE "SequencerDisplayNode"
@@ -484,6 +485,14 @@ void FSequencerDisplayNode::BuildContextMenu(FMenuBuilder& MenuBuilder)
 			EUserInterfaceActionType::ToggleButton
 		);
 
+
+		// Add cut, copy and paste functions to the tracks
+		MenuBuilder.AddMenuEntry(FGenericCommands::Get().Cut);
+
+		MenuBuilder.AddMenuEntry(FGenericCommands::Get().Copy);
+		
+		MenuBuilder.AddMenuEntry(FGenericCommands::Get().Paste);
+		
 		MenuBuilder.AddMenuEntry(
 			LOCTEXT("DeleteNode", "Delete"),
 			LOCTEXT("DeleteNodeTooltip", "Delete this or selected tracks"),
diff --git a/Engine/Source/Editor/Sequencer/Private/SSequencer.cpp b/Engine/Source/Editor/Sequencer/Private/SSequencer.cpp
--- a/Engine/Source/Editor/Sequencer/Private/SSequencer.cpp
+++ b/Engine/Source/Editor/Sequencer/Private/SSequencer.cpp
@@ -513,7 +513,7 @@ PRAGMA_ENABLE_OPTIMIZATION
 
 void SSequencer::BindCommands(TSharedRef<FUICommandList> SequencerCommandBindings)
 {
-	auto CanPaste = [this]{
+	auto CanPasteFromHistory = [this]{
 		if (!HasFocusedDescendants() && !HasKeyboardFocus())
 		{
 			return false;
@@ -524,14 +524,14 @@ void SSequencer::BindCommands(TSharedRef<FUICommandList> SequencerCommandBinding
 	
 	SequencerCommandBindings->MapAction(
 		FGenericCommands::Get().Paste,
-		FExecuteAction::CreateSP(this, &SSequencer::Paste),
-		FCanExecuteAction::CreateLambda(CanPaste)
+		FExecuteAction::CreateSP(this, &SSequencer::OnPaste),
+		FCanExecuteAction::CreateSP(this, &SSequencer::CanPaste)
 	);
 
 	SequencerCommandBindings->MapAction(
 		FSequencerCommands::Get().PasteFromHistory,
 		FExecuteAction::CreateSP(this, &SSequencer::PasteFromHistory),
-		FCanExecuteAction::CreateLambda(CanPaste)
+		FCanExecuteAction::CreateLambda(CanPasteFromHistory)
 	);
 
 	SequencerCommandBindings->MapAction(
@@ -1828,7 +1828,83 @@ FPasteContextMenuArgs SSequencer::GeneratePasteArgs(float PasteAtTime, TSharedPt
 	return FPasteContextMenuArgs::PasteInto(MoveTemp(PasteIntoNodes), PasteAtTime, Clipboard);
 }
 
-void SSequencer::Paste()
+void SSequencer::OnPaste()
+{
+	TSharedPtr<FSequencer> Sequencer = SequencerPtr.Pin();
+	TSet<TSharedRef<FSequencerDisplayNode>> SelectedNodes = Sequencer->GetSelection().GetSelectedOutlinerNodes();
+	if (SelectedNodes.Num() == 0)
+	{
+		OpenPasteMenu();
+	}
+	else
+	{
+		PasteTracks();
+	}
+}
+
+bool SSequencer::CanPaste()
+{
+	TSharedPtr<FSequencer> Sequencer = SequencerPtr.Pin();
+	TSet<TSharedRef<FSequencerDisplayNode>> SelectedNodes = Sequencer->GetSelection().GetSelectedOutlinerNodes();
+	if (SelectedNodes.Num() != 0)
+	{
+		FString TexttoImport;
+		FPlatformMisc::ClipboardPaste(TexttoImport);
+
+		TArray<UMovieSceneTrack*> ImportedTrack;
+		Sequencer->ImportTracksFromText(TexttoImport, ImportedTrack);
+		if (ImportedTrack.Num() == 0)
+		{
+			return false;
+		}
+
+		for (TSharedRef<FSequencerDisplayNode> Node : SelectedNodes)
+		{
+			if (Node->GetType() == ESequencerNode::Object)
+			{
+				return true;
+			}
+		}
+		return false;
+	}
+	else
+	{
+		if (!HasFocusedDescendants() && !HasKeyboardFocus())
+		{
+			return false;
+		}
+		return SequencerPtr.Pin()->GetClipboardStack().Num() != 0;
+	}
+}
+
+void SSequencer::PasteTracks()
+{
+	TSharedPtr<FSequencer> Sequencer = SequencerPtr.Pin();
+
+	TSet<TSharedRef<FSequencerDisplayNode>> SelectedNodes = Sequencer->GetSelection().GetSelectedOutlinerNodes();
+	if (SelectedNodes.Num() == 0)
+	{
+		return;
+	}
+
+	TArray<TSharedPtr<FSequencerObjectBindingNode>> ObjectNodes;
+	for (TSharedRef<FSequencerDisplayNode> Node : SelectedNodes)
+	{
+		if (Node->GetType() != ESequencerNode::Object)
+		{
+			continue;
+		}
+
+		TSharedPtr<FSequencerObjectBindingNode> ObjectNode = StaticCastSharedRef<FSequencerObjectBindingNode>(Node);
+		if (ObjectNode.IsValid())
+		{
+			ObjectNodes.Add(ObjectNode);
+		}
+	}
+	Sequencer->PasteCopiedTracks(ObjectNodes);
+}
+
+void SSequencer::OpenPasteMenu()
 {
 	TSharedPtr<FPasteContextMenu> ContextMenu;
 
diff --git a/Engine/Source/Editor/Sequencer/Private/SSequencer.h b/Engine/Source/Editor/Sequencer/Private/SSequencer.h
--- a/Engine/Source/Editor/Sequencer/Private/SSequencer.h
+++ b/Engine/Source/Editor/Sequencer/Private/SSequencer.h
@@ -389,9 +389,15 @@ private:
 	void OnSequenceInstanceActivated( FMovieSceneSequenceInstance& ActiveInstance );
 
 public:
+	/** On Paste Command */
+	void OnPaste();
+	bool CanPaste();
+
+	/** Handle Track Paste */
+	void PasteTracks();
 
 	/** Open the paste menu */
-	void Paste();
+	void OpenPasteMenu();
 	
 	/** Open the paste from history menu */
 	void PasteFromHistory();
diff --git a/Engine/Source/Editor/Sequencer/Private/Sequencer.cpp b/Engine/Source/Editor/Sequencer/Private/Sequencer.cpp
--- a/Engine/Source/Editor/Sequencer/Private/Sequencer.cpp
+++ b/Engine/Source/Editor/Sequencer/Private/Sequencer.cpp
@@ -76,6 +76,8 @@
 #include "MainFrame.h"
 #include "DesktopPlatformModule.h"
 #include "FbxExporter.h"
+#include "UnrealExporter.h"
+#include "Factories.h"
 
 #define LOCTEXT_NAMESPACE "Sequencer"
 
@@ -263,6 +265,7 @@ FSequencer::FSequencer()
 {
 	Selection.GetOnOutlinerNodeSelectionChanged().AddRaw(this, &FSequencer::OnSelectedOutlinerNodesChanged);
 	Selection.GetOnNodesWithSelectedKeysOrSectionsChanged().AddRaw(this, &FSequencer::OnSelectedOutlinerNodesChanged);
+	Selection.GetOnOutlinerNodeSelectionChangedObjectGuids().AddRaw(this, &FSequencer::OnSelectedOutlinerNodesChanged);
 }
 
 
@@ -788,7 +791,7 @@ void FSequencer::GetKeysFromSelection(TUniquePtr<ISequencerKeyCollection>& KeyCo
 		SelectedNodes.Add(&Node.Get());
 	}
 
-	// Anything within .5 pixel's worth of time is a duplicate as far as we're concerened
+	// Anything within .5 pixel's worth of time is a duplicate as far as we're concerned
 	FVirtualTrackArea TrackArea = SequencerWidget->GetVirtualTrackArea();;
 	const float DuplicateThreshold = (TrackArea.PixelToTime(0.f) - TrackArea.PixelToTime(1.f)) * .5f;
 
@@ -4016,6 +4019,7 @@ void FSequencer::OnSelectedOutlinerNodesChanged()
 		SynchronizeExternalSelectionWithSequencerSelection();
 		bUpdatingExternalSelection = false;
 	}
+	OnSelectionChangedObjectGuidsDelegate.Broadcast(Selection.GetBoundObjectsGuids());
 }
 
 
@@ -4514,6 +4518,181 @@ void FSequencer::DeleteSelectedNodes()
 }
 
 
+void FSequencer::CopySelectedTracks(TArray<TSharedPtr<FSequencerTrackNode>>& TrackNodes)
+{
+	TArray<UMovieSceneTrack*> TracksToCopy;
+	for (TSharedPtr<FSequencerTrackNode> TrackNode : TrackNodes)
+	{
+		TracksToCopy.Add(TrackNode->GetTrack());
+	}
+
+	FString ExportedText;
+	FSequencer::ExportTracksToText(TracksToCopy, /*out*/ ExportedText);
+	FPlatformMisc::ClipboardCopy(*ExportedText);
+}
+
+
+void FSequencer::ExportTracksToText(TArray<UMovieSceneTrack*> TracksToExport, FString& ExportedText)
+{
+	// Clear the mark state for saving.
+	UnMarkAllObjects(EObjectMark(OBJECTMARK_TagExp | OBJECTMARK_TagImp));
+
+	FStringOutputDevice Archive;
+	const FExportObjectInnerContext Context;
+
+	// Export each of the selected nodes
+	UObject* LastOuter = nullptr;
+
+	for (UMovieSceneTrack* TrackToExport : TracksToExport)
+	{
+		// The nodes should all be from the same scope
+		UObject* ThisOuter = TrackToExport->GetOuter();
+		check((LastOuter == ThisOuter) || (LastOuter == nullptr));
+		LastOuter = ThisOuter;
+
+		UExporter::ExportToOutputDevice(&Context, TrackToExport, nullptr, Archive, TEXT("copy"), 0, PPF_ExportsNotFullyQualified | PPF_Copy | PPF_Delimited, false, ThisOuter);
+	}
+
+	ExportedText = Archive;
+}
+
+void FSequencer::PasteCopiedTracks(TArray<TSharedPtr<FSequencerObjectBindingNode>>& ObjectNodes)
+{
+	FScopedTransaction Transaction(FGenericCommands::Get().Paste->GetDescription());
+	// Grab the text to paste from the clipboard
+	FString TextToImport;
+	FPlatformMisc::ClipboardPaste(TextToImport);
+
+	TArray<UMovieSceneTrack*> ImportedTracks;
+	FSequencer::ImportTracksFromText(TextToImport, ImportedTracks);
+
+	if (ImportedTracks.Num() == 0)
+	{
+		Transaction.Cancel();
+		return;
+	}
+
+	for (TSharedPtr<FSequencerObjectBindingNode> ObjectNode : ObjectNodes)
+	{
+		FGuid ObjectGuid = ObjectNode->GetObjectBinding();
+		UObject* Context = nullptr;
+		UObject* PossessableObject = GetFocusedMovieSceneSequence()->FindPossessableObject(ObjectGuid, Context);
+
+		if (!PossessableObject)
+		{
+			FNotificationInfo Info(LOCTEXT("InvalidTargetObject", "Invalid Target Object"));
+			Info.FadeInDuration = 0.1f;
+			Info.FadeOutDuration = 0.5f;
+			Info.ExpireDuration = 2.5f;
+			auto NotificationItem = FSlateNotificationManager::Get().AddNotification(Info);
+
+			NotificationItem->SetCompletionState(SNotificationItem::CS_Success);
+			NotificationItem->ExpireAndFadeout();
+			continue;
+		}
+
+		for (UMovieSceneTrack* ImportedTrack : ImportedTracks)
+		{
+			UMovieScenePropertyTrack* CastedPropertyTrack = Cast<UMovieScenePropertyTrack>(ImportedTrack);
+
+			if (!CastedPropertyTrack)
+			{
+				FNotificationInfo Info(LOCTEXT("CopiedTrackIsNotPropertyTrack", "Can't Paste: Copied Track isn't a Property Track"));
+				Info.FadeInDuration = 0.1f;
+				Info.FadeOutDuration = 0.5f;
+				Info.ExpireDuration = 2.5f;
+				auto NotificationItem = FSlateNotificationManager::Get().AddNotification(Info);
+
+				NotificationItem->SetCompletionState(SNotificationItem::CS_Success);
+				NotificationItem->ExpireAndFadeout();
+
+				continue;
+			}
+
+			FString PropertyName = CastedPropertyTrack->GetPropertyName().ToString();
+			PropertyName.RemoveFromStart("b", ESearchCase::CaseSensitive);
+			FString FunctionName = "Set" + PropertyName;
+			if (PossessableObject->FindFunction(FName(*FunctionName)))
+			{
+				if (!GetFocusedMovieSceneSequence()->GetMovieScene()->AddGivenTrack(ImportedTrack, ObjectGuid))
+				{
+					FNotificationInfo Info(LOCTEXT("PropertyAlreadyBound", "Can't Paste: Property already bound to a track"));
+					Info.FadeInDuration = 0.1f;
+					Info.FadeOutDuration = 0.5f;
+					Info.ExpireDuration = 2.5f;
+					auto NotificationItem = FSlateNotificationManager::Get().AddNotification(Info);
+
+					NotificationItem->SetCompletionState(SNotificationItem::CS_Success);
+					NotificationItem->ExpireAndFadeout();
+
+					continue;
+				}
+				else
+				{
+					NotifyMovieSceneDataChanged(EMovieSceneDataChangeType::MovieSceneStructureItemsChanged);
+				}
+			}
+			else
+			{
+				FNotificationInfo Info(LOCTEXT("PropertyNotFound", "Can't Paste: Target Object doesn't have this property"));
+				Info.FadeInDuration = 0.1f;
+				Info.FadeOutDuration = 0.5f;
+				Info.ExpireDuration = 2.5f;
+				auto NotificationItem = FSlateNotificationManager::Get().AddNotification(Info);
+
+				NotificationItem->SetCompletionState(SNotificationItem::CS_Success);
+				NotificationItem->ExpireAndFadeout();
+
+				continue;
+			}
+		}
+	}
+	
+}
+
+
+class FTrackObjectTextFactory : public FCustomizableTextObjectFactory
+{
+public:
+	FTrackObjectTextFactory()
+		: FCustomizableTextObjectFactory(GWarn)
+	{
+	}
+
+	// FCustomizableTextObjectFactory implementation
+	virtual bool CanCreateClass(UClass* ObjectClass, bool& bOmitSubObjs) const override
+	{
+		return true;
+	}
+	
+
+	virtual void ProcessConstructedObject(UObject* NewObject) override
+	{
+		check(NewObject);
+
+		NewTracks.Add(Cast<UMovieSceneTrack>(NewObject));
+	}
+
+public:
+	TArray<UMovieSceneTrack*> NewTracks;
+};
+
+void FSequencer::ImportTracksFromText(const FString& TextToImport, /*out*/ TArray<UMovieSceneTrack*>& ImportedTracks)
+{
+	UPackage* TempPackage = NewObject<UPackage>(nullptr, TEXT("/Engine/Sequencer/Editor/Transient"), RF_Transient);
+	TempPackage->AddToRoot();
+
+	// Turn the text buffer into objects
+	FTrackObjectTextFactory Factory;
+	Factory.ProcessBuffer(TempPackage, RF_Transactional, TextToImport);
+
+	ImportedTracks = Factory.NewTracks;
+
+	// Remove the temp package from the root now that it has served its purpose
+	TempPackage->RemoveFromRoot();
+}
+
+
 void FSequencer::ToggleNodeActive()
 {
 	bool bIsActive = !IsNodeActive();
@@ -5257,14 +5436,94 @@ void FSequencer::SelectTrackKeys(TWeakObjectPtr<UMovieSceneSection> Section, flo
 
 TArray<TSharedPtr<FMovieSceneClipboard>> GClipboardStack;
 
+void FSequencer::CopySelection()
+{
+	if (Selection.GetSelectedKeys().Num() == 0)
+	{
+		TArray<TSharedPtr<FSequencerTrackNode>> TracksToCopy;
+		TSet<TSharedRef<FSequencerDisplayNode>> SelectedNodes = Selection.GetNodesWithSelectedKeysOrSections();
+		if (SelectedNodes.Num() == 0)
+		{
+			SelectedNodes = Selection.GetSelectedOutlinerNodes();
+		}
+		for (TSharedRef<FSequencerDisplayNode> Node : SelectedNodes)
+		{
+			if (Node->GetType() != ESequencerNode::Track)
+			{
+				continue;
+			}
+
+			TSharedPtr<FSequencerTrackNode> TrackNode = StaticCastSharedRef<FSequencerTrackNode>(Node);
+			if (TrackNode.IsValid())
+			{
+				TracksToCopy.Add(TrackNode);
+			}
+		}
 
-void FSequencer::CopySelectedKeys()
+		// Make a empty clipboard if the stack is empty
+		if (GClipboardStack.Num() == 0)
+		{
+			TSharedRef<FMovieSceneClipboard> NullClipboard = MakeShareable(new FMovieSceneClipboard());
+			GClipboardStack.Push(NullClipboard);
+		}
+		CopySelectedTracks(TracksToCopy);
+	}
+	else
+	{
+		CopySelectedKeys();
+	}
+}
+
+
+void FSequencer::CutSelection()
 {
+	FScopedTransaction CutSelectionTransaction(LOCTEXT("CutSelection_Transaction", "Cut Selection(s)"));
 	if (Selection.GetSelectedKeys().Num() == 0)
 	{
-		return;
+		TArray<TSharedPtr<FSequencerTrackNode>> TracksToCopy;
+		TSet<TSharedRef<FSequencerDisplayNode>> SelectedNodes = Selection.GetNodesWithSelectedKeysOrSections();
+		if (SelectedNodes.Num() == 0)
+		{
+			SelectedNodes = Selection.GetSelectedOutlinerNodes();
+		}
+		for (TSharedRef<FSequencerDisplayNode> Node : SelectedNodes)
+		{
+			if (Node->GetType() != ESequencerNode::Track)
+			{
+				FNotificationInfo Info(LOCTEXT("InvalidCut", "Warning: One of the selected node is not a track node"));
+				Info.FadeInDuration = 0.1f;
+				Info.FadeOutDuration = 0.5f;
+				Info.ExpireDuration = 2.5f;
+				auto NotificationItem = FSlateNotificationManager::Get().AddNotification(Info);
+
+				NotificationItem->SetCompletionState(SNotificationItem::CS_Success);
+				NotificationItem->ExpireAndFadeout();
+				return;
+			}
+
+			TSharedPtr<FSequencerTrackNode> TrackNode = StaticCastSharedRef<FSequencerTrackNode>(Node);
+			if (TrackNode.IsValid())
+			{
+				TracksToCopy.Add(TrackNode);
+			}
+		}
+		// Make a empty clipboard if the stack is empty
+		if (GClipboardStack.Num() == 0)
+		{
+			TSharedRef<FMovieSceneClipboard> NullClipboard = MakeShareable(new FMovieSceneClipboard());
+			GClipboardStack.Push(NullClipboard);
+		}
+		CopySelectedTracks(TracksToCopy);
+		DeleteSelectedItems();
 	}
+	else
+	{
+		CutSelectedKeys();
+	}
+}
 
+void FSequencer::CopySelectedKeys()
+{
 	TOptional<float> CopyRelativeTo;
 	
 	// Copy relative to the current key hotspot, if applicable
@@ -5313,7 +5572,7 @@ void FSequencer::CopySelectedKeys()
 
 void FSequencer::CutSelectedKeys()
 {
-	FScopedTransaction CutKeysTransaction( LOCTEXT("CutSelectedKeys_Transaction", "Cut Selected Keys") );
+	FScopedTransaction CutSelectedKeysTransaction(LOCTEXT("CutSelectedKeys_Transaction", "Cut Selected keys"));
 	CopySelectedKeys();
 	DeleteSelectedKeys();
 }
@@ -6123,6 +6382,31 @@ void FSequencer::BindCommands()
 		FIsActionChecked::CreateLambda( [this]{ return Settings->GetLinkCurveEditorTimeRange(); } ) );
 
 	auto CanCutOrCopy = [this]{
+		// For copy tracks
+		TArray<TSharedPtr<FSequencerTrackNode>> TracksToCopy;
+		TSet<TSharedRef<FSequencerDisplayNode>> SelectedNodes = Selection.GetNodesWithSelectedKeysOrSections();
+		// If this is empty then we are selecting display nodes
+		if (SelectedNodes.Num() == 0)
+		{
+			SelectedNodes = Selection.GetSelectedOutlinerNodes();
+			for (TSharedRef<FSequencerDisplayNode> Node : SelectedNodes)
+			{
+				if (Node->GetType() == ESequencerNode::Track)
+				{
+					// if contains one node that can be copied we allow the action
+					// later on we will filter out the invalid nodes in CopySelection() or CutSelection()
+					return true;
+				}
+				else if (Node->GetParent().IsValid() && Node->GetParent()->GetType() == ESequencerNode::Track)
+				{
+					// Although copying only the child nodes (ex. translation) is not allowed, we still show the copy & cut button
+					// so that users are not misled and can achieve this in copy/cut the parent node (ex. transform)
+					return true;
+				}
+			}
+			return false;
+		}
+
 		UMovieSceneTrack* Track = nullptr;
 		for (FSequencerSelectedKey Key : Selection.GetSelectedKeys())
 		{
@@ -6140,13 +6424,13 @@ void FSequencer::BindCommands()
 
 	SequencerCommandBindings->MapAction(
 		FGenericCommands::Get().Cut,
-		FExecuteAction::CreateSP(this, &FSequencer::CutSelectedKeys),
+		FExecuteAction::CreateSP(this, &FSequencer::CutSelection),
 		FCanExecuteAction::CreateLambda(CanCutOrCopy)
 	);
 
 	SequencerCommandBindings->MapAction(
 		FGenericCommands::Get().Copy,
-		FExecuteAction::CreateSP(this, &FSequencer::CopySelectedKeys),
+		FExecuteAction::CreateSP(this, &FSequencer::CopySelection),
 		FCanExecuteAction::CreateLambda(CanCutOrCopy)
 	);
 
diff --git a/Engine/Source/Editor/Sequencer/Private/Sequencer.h b/Engine/Source/Editor/Sequencer/Private/Sequencer.h
--- a/Engine/Source/Editor/Sequencer/Private/Sequencer.h
+++ b/Engine/Source/Editor/Sequencer/Private/Sequencer.h
@@ -409,6 +409,14 @@ public:
 	void DeleteNode(TSharedRef<FSequencerDisplayNode> NodeToBeDeleted);
 	void DeleteSelectedNodes();
 
+	/** Called when a user executes the copy track menu item */
+	void CopySelectedTracks(TArray<TSharedPtr<FSequencerTrackNode>>& TrackNodes);
+	void ExportTracksToText(TArray<UMovieSceneTrack*> TrackToExport, /*out*/ FString& ExportedText);
+
+	/** Called when a user executes the paste track menu item */
+	void PasteCopiedTracks(TArray<TSharedPtr<FSequencerObjectBindingNode>>& ObjectNodes);
+	void ImportTracksFromText(const FString& TextToImport, /*out*/ TArray<UMovieSceneTrack*>& ImportedTrack);
+
 	/** Called when a user executes the active node menu item */
 	void ToggleNodeActive();
 	bool IsNodeActive() const;
@@ -438,6 +446,12 @@ public:
 
 public:
 
+	/** Copy the selection, whether it's keys or tracks */
+	void CopySelection();
+
+	/** Cut the selection, whether it's keys or tracks */
+	void CutSelection();
+
 	/** Copy the selected keys to the clipboard */
 	void CopySelectedKeys();
 
@@ -561,6 +575,7 @@ public:
 	virtual void NotifyMapChanged(UWorld* NewWorld, EMapChangeType MapChangeType) override;
 	virtual FOnGlobalTimeChanged& OnGlobalTimeChanged() override { return OnGlobalTimeChangedDelegate; }
 	virtual FOnMovieSceneDataChanged& OnMovieSceneDataChanged() override { return OnMovieSceneDataChangedDelegate; }
+	virtual FOnSelectionChangedObjectGuids& GetSelectionChangedObjectGuids() override { return OnSelectionChangedObjectGuidsDelegate; }
 	virtual FGuid CreateBinding(UObject& InObject, const FString& InName) override;
 	virtual UObject* GetPlaybackContext() const override;
 	virtual TArray<UObject*> GetEventContexts() const override;
@@ -955,6 +970,9 @@ private:
 	/** A delegate which is called any time the movie scene data is changed. */
 	FOnMovieSceneDataChanged OnMovieSceneDataChangedDelegate;
 
+	/** A delegate which is called any time the sequencer selection changes. */
+	FOnSelectionChangedObjectGuids OnSelectionChangedObjectGuidsDelegate;
+
 	/** A map of all the transport controls to viewports that this sequencer has made */
 	TMap< TSharedPtr<class ILevelViewport>, TSharedPtr<class SWidget> > TransportControls;
 
@@ -969,7 +987,7 @@ private:
 	FLevelEditorModule::FLevelEditorMenuExtender ViewMenuExtender;
 	FDelegateHandle LevelEditorExtenderDelegateHandle;
 
-	/** When true the sequencer selection is being updated from changes to the external seleciton. */
+	/** When true the sequencer selection is being updated from changes to the external selection. */
 	bool bUpdatingSequencerSelection;
 
 	/** When true the external selection is being updated from changes to the sequencer selection. */
diff --git a/Engine/Source/Editor/Sequencer/Private/SequencerContextMenus.cpp b/Engine/Source/Editor/Sequencer/Private/SequencerContextMenus.cpp
--- a/Engine/Source/Editor/Sequencer/Private/SequencerContextMenus.cpp
+++ b/Engine/Source/Editor/Sequencer/Private/SequencerContextMenus.cpp
@@ -328,6 +328,13 @@ void FSectionContextMenu::PopulateMenu(FMenuBuilder& MenuBuilder)
 				FSlateIcon(),
 				FUIAction(FExecuteAction::CreateLambda([=]{ return Shared->SelectAllKeys(); }))
 			);
+
+			MenuBuilder.AddMenuEntry(
+				LOCTEXT("CopyAllKeys", "Copy All Keys"),
+				LOCTEXT("CopyAllKeysTooltip", "Copy all keys in section"),
+				FSlateIcon(),
+				FUIAction(FExecuteAction::CreateLambda([=] { return Shared->CopyAllKeys(); }))
+			);
 		}
 
 		MenuBuilder.AddSubMenu(
@@ -577,6 +584,11 @@ void FSectionContextMenu::SelectAllKeys()
 	}
 }
 
+void FSectionContextMenu::CopyAllKeys()
+{
+	SelectAllKeys();
+	Sequencer->CopySelectedKeys();
+}
 
 void FSectionContextMenu::TogglePrimeForRecording() const
 {
diff --git a/Engine/Source/Editor/Sequencer/Private/SequencerContextMenus.h b/Engine/Source/Editor/Sequencer/Private/SequencerContextMenus.h
--- a/Engine/Source/Editor/Sequencer/Private/SequencerContextMenus.h
+++ b/Engine/Source/Editor/Sequencer/Private/SequencerContextMenus.h
@@ -41,6 +41,8 @@ private:
 
 	void SelectAllKeys();
 
+	void CopyAllKeys();
+
 	bool CanSelectAllKeys() const;
 
 	void TogglePrimeForRecording() const;
diff --git a/Engine/Source/Editor/Sequencer/Private/SequencerSelection.cpp b/Engine/Source/Editor/Sequencer/Private/SequencerSelection.cpp
--- a/Engine/Source/Editor/Sequencer/Private/SequencerSelection.cpp
+++ b/Engine/Source/Editor/Sequencer/Private/SequencerSelection.cpp
@@ -50,12 +50,55 @@ FSequencerSelection::FOnSelectionChanged& FSequencerSelection::GetOnNodesWithSel
 	return OnNodesWithSelectedKeysOrSectionsChanged;
 }
 
+FSequencerSelection::FOnSelectionChangedObjectGuids& FSequencerSelection::GetOnOutlinerNodeSelectionChangedObjectGuids()
+{
+	return OnOutlinerNodeSelectionChangedObjectGuids;
+}
+
+TArray<FGuid> FSequencerSelection::GetBoundObjectsGuids()
+{
+	TArray<FGuid> OutGuids;
+	TSet<TSharedRef<FSequencerDisplayNode>> SelectedNodes = GetNodesWithSelectedKeysOrSections();
+	if (SelectedNodes.Num() == 0)
+	{
+		SelectedNodes = GetSelectedOutlinerNodes();
+	}
+	for (TSharedRef<FSequencerDisplayNode> Node : SelectedNodes)
+	{
+		TSharedPtr<FSequencerObjectBindingNode> ObjectNode;
+		if (Node->GetType() == ESequencerNode::Object)
+		{
+			ObjectNode = StaticCastSharedRef<FSequencerObjectBindingNode>(Node);
+		}
+		else
+		{
+			TSharedPtr<FSequencerDisplayNode> ParentNode = Node->GetParent();
+			if (ParentNode.IsValid())
+			{
+				while (ParentNode->GetParent().IsValid() && ParentNode->GetType() != ESequencerNode::Object)
+				{
+					ParentNode = ParentNode->GetParent();
+				}
+				ObjectNode = StaticCastSharedPtr<FSequencerObjectBindingNode>(ParentNode);
+			}
+		}
+
+		if (ObjectNode.IsValid())
+		{
+			OutGuids.Add(ObjectNode->GetObjectBinding());
+		}
+	}
+
+	return OutGuids;
+}
+
 void FSequencerSelection::AddToSelection(FSequencerSelectedKey Key)
 {
 	SelectedKeys.Add(Key);
 	if ( IsBroadcasting() )
 	{
 		OnKeySelectionChanged.Broadcast();
+		OnOutlinerNodeSelectionChangedObjectGuids.Broadcast();
 	}
 
 	// Deselect any outliner nodes that aren't within the trunk of this key
@@ -73,6 +116,7 @@ void FSequencerSelection::AddToSelection(UMovieSceneSection* Section)
 	if ( IsBroadcasting() )
 	{
 		OnSectionSelectionChanged.Broadcast();
+		OnOutlinerNodeSelectionChangedObjectGuids.Broadcast();
 	}
 
 	// Deselect any outliner nodes that aren't within the trunk of this section
@@ -90,6 +134,7 @@ void FSequencerSelection::AddToSelection(TSharedRef<FSequencerDisplayNode> Outli
 	if ( IsBroadcasting() )
 	{
 		OnOutlinerNodeSelectionChanged.Broadcast();
+		OnOutlinerNodeSelectionChangedObjectGuids.Broadcast();
 	}
 	EmptySelectedKeys();
 	EmptySelectedSections();
@@ -102,6 +147,7 @@ void FSequencerSelection::AddToNodesWithSelectedKeysOrSections(TSharedRef<FSeque
 	if ( IsBroadcasting() )
 	{
 		OnNodesWithSelectedKeysOrSectionsChanged.Broadcast();
+		OnOutlinerNodeSelectionChangedObjectGuids.Broadcast();
 	}
 }
 
diff --git a/Engine/Source/Editor/Sequencer/Private/SequencerSelection.h b/Engine/Source/Editor/Sequencer/Private/SequencerSelection.h
--- a/Engine/Source/Editor/Sequencer/Private/SequencerSelection.h
+++ b/Engine/Source/Editor/Sequencer/Private/SequencerSelection.h
@@ -16,6 +16,8 @@ public:
 
 	DECLARE_MULTICAST_DELEGATE(FOnSelectionChanged)
 
+	DECLARE_MULTICAST_DELEGATE(FOnSelectionChangedObjectGuids)
+
 public:
 
 	FSequencerSelection();
@@ -98,6 +100,12 @@ public:
 	/** Gets a multicast delegate which is called when the outliner node with selected keys or sections changes. */
 	FOnSelectionChanged& GetOnNodesWithSelectedKeysOrSectionsChanged();
 
+	/** Gets a multicast delegate with an array of FGuid of bound objects which is called when the outliner node selection changes. */
+	FOnSelectionChangedObjectGuids& GetOnOutlinerNodeSelectionChangedObjectGuids();
+
+	/** Helper function to get an array of FGuid of bound objects on return */
+	TArray<FGuid> GetBoundObjectsGuids();
+
 	/** Suspend the broadcast of selection change notifications.  */
 	void SuspendBroadcast();
 
@@ -126,6 +134,8 @@ private:
 	FOnSelectionChanged OnOutlinerNodeSelectionChanged;
 	FOnSelectionChanged OnNodesWithSelectedKeysOrSectionsChanged;
 
+	FOnSelectionChangedObjectGuids OnOutlinerNodeSelectionChangedObjectGuids;
+
 	/** The number of times the broadcasting of selection change notifications has been suspended. */
 	int32 SuspendBroadcastCount;
 
diff --git a/Engine/Source/Editor/Sequencer/Public/ISequencer.h b/Engine/Source/Editor/Sequencer/Public/ISequencer.h
--- a/Engine/Source/Editor/Sequencer/Public/ISequencer.h
+++ b/Engine/Source/Editor/Sequencer/Public/ISequencer.h
@@ -99,6 +99,8 @@ public:
 	
 	DECLARE_MULTICAST_DELEGATE(FOnGlobalTimeChanged);
 	DECLARE_MULTICAST_DELEGATE(FOnMovieSceneDataChanged);
+	
+	DECLARE_MULTICAST_DELEGATE_OneParam(FOnSelectionChangedObjectGuids, TArray<FGuid> /*Object*/)
 
 public:
 
@@ -353,6 +355,9 @@ public:
 	/** Gets a multicast delegate which is executed whenever the movie scene data is changed. */
 	virtual FOnMovieSceneDataChanged& OnMovieSceneDataChanged() = 0;
 
+	/** Gets a multicast delegate with an array of FGuid of bound objects which is called when the outliner node selection changes. */
+	virtual FOnSelectionChangedObjectGuids& GetSelectionChangedObjectGuids() = 0;
+
 	/** @return a numeric type interface that will parse and display numbers as frames and times correctly */
 	virtual TSharedRef<INumericTypeInterface<float>> GetNumericTypeInterface() = 0;
 
diff --git a/Engine/Source/Editor/SourceControlWindows/Private/SSourceControlHistory.cpp b/Engine/Source/Editor/SourceControlWindows/Private/SSourceControlHistory.cpp
--- a/Engine/Source/Editor/SourceControlWindows/Private/SSourceControlHistory.cpp
+++ b/Engine/Source/Editor/SourceControlWindows/Private/SSourceControlHistory.cpp
@@ -1205,9 +1205,9 @@ private:
 	}
 
 	/**
-	 * Checks to see if the selected history-tree items can be diff'd against each other.
+	 * Checks to see if the selected history-tree items can be diffed against each other.
 	 * 
-	 * @return True if the selected items can be diff'd, false if not.
+	 * @return True if the selected items can be diffed, false if not.
 	 */
 	bool CanDiffSelected() const
 	{
diff --git a/Engine/Source/Editor/TranslationEditor/Private/TranslationDataManager.cpp b/Engine/Source/Editor/TranslationEditor/Private/TranslationDataManager.cpp
--- a/Engine/Source/Editor/TranslationEditor/Private/TranslationDataManager.cpp
+++ b/Engine/Source/Editor/TranslationEditor/Private/TranslationDataManager.cpp
@@ -7,8 +7,7 @@
 #include "ISourceControlModule.h"
 #include "MessageLog.h"
 #include "TextLocalizationManager.h"
-#include "JsonInternationalizationArchiveSerializer.h"
-#include "JsonInternationalizationManifestSerializer.h"
+#include "TextLocalizationResourceGenerator.h"
 #include "SNotificationList.h"
 #include "NotificationManager.h"
 #include "PortableObjectFormatDOM.h"
@@ -16,12 +15,58 @@
 #include "LocalizationModule.h"
 #include "LocalizationTargetTypes.h"
 #include "LocalizationConfigurationScript.h"
+#include "JsonInternationalizationArchiveSerializer.h"
+#include "JsonInternationalizationManifestSerializer.h"
 
 
 DEFINE_LOG_CATEGORY_STATIC(LogTranslationEditor, Log, All);
 
 #define LOCTEXT_NAMESPACE "TranslationDataManager"
 
+struct FLocTextIdentity
+{
+public:
+	FLocTextIdentity(FString InNamespace, FString InKey)
+		: Namespace(MoveTemp(InNamespace))
+		, Key(MoveTemp(InKey))
+		, Hash(0)
+	{
+		Hash = FCrc::StrCrc32(*Namespace, Hash);
+		Hash = FCrc::StrCrc32(*Key, Hash);
+	}
+
+	FORCEINLINE const FString& GetNamespace() const
+	{
+		return Namespace;
+	}
+
+	FORCEINLINE const FString& GetKey() const
+	{
+		return Key;
+	}
+
+	FORCEINLINE bool operator==(const FLocTextIdentity& Other) const
+	{
+		return Namespace.Equals(Other.Namespace, ESearchCase::CaseSensitive)
+			&& Key.Equals(Other.Key, ESearchCase::CaseSensitive);
+	}
+
+	FORCEINLINE bool operator!=(const FLocTextIdentity& Other) const
+	{
+		return !Namespace.Equals(Other.Namespace, ESearchCase::CaseSensitive)
+			|| !Key.Equals(Other.Key, ESearchCase::CaseSensitive);
+	}
+
+	friend inline uint32 GetTypeHash(const FLocTextIdentity& Id)
+	{
+		return Id.Hash;
+	}
+
+private:
+	FString Namespace;
+	FString Key;
+	uint32 Hash;
+};
 
 FTranslationDataManager::FTranslationDataManager( const FString& InManifestFilePath, const FString& InNativeArchiveFilePath, const FString& InArchiveFilePath )
 	: OpenedManifestFilePath(InManifestFilePath)
@@ -85,12 +130,12 @@ void FTranslationDataManager::Initialize()
 
 			GWarn->BeginSlowTask(LOCTEXT("LoadingCurrentManifest", "Loading Entries from Current Translation Manifest..."), true);
 
-			// Get all manifest entries by source text (same source text in multiple contexts will only show up once, unless they have unique key metadata)
+			// Get all manifest entries by source text...
 			for (auto ManifestItr = ManifestAtHeadRevision->GetEntriesBySourceTextIterator(); ManifestItr; ++ManifestItr, ++NumManifestEntriesParsed)
 			{
 				GWarn->StatusUpdate(NumManifestEntriesParsed, ManifestEntriesCount, FText::Format(LOCTEXT("LoadingCurrentManifestEntries", "Loading Entry {0} of {1} from Current Translation Manifest..."), FText::AsNumber(NumManifestEntriesParsed), FText::AsNumber(ManifestEntriesCount)));
 				const TSharedRef<FManifestEntry> ManifestEntry = ManifestItr.Value();
-				TMap< TSharedPtr<FLocMetadataObject>, UTranslationUnit* > KeyMetaDataToTranslationUnitMap;
+				TMap< FLocTextIdentity, UTranslationUnit* > IdentityToTranslationUnitMap;
 
 				for(auto ContextIter( ManifestEntry->Contexts.CreateConstIterator() ); ContextIter; ++ContextIter)
 				{
@@ -100,8 +145,8 @@ void FTranslationDataManager::Initialize()
 					ContextInfo.Context = AContext.SourceLocation;
 					ContextInfo.Key = AContext.Key;
 
-					// Make sure we have a unique translation unit for each unique key metadata object.
-					UTranslationUnit*& TranslationUnit = KeyMetaDataToTranslationUnitMap.FindOrAdd(AContext.KeyMetadataObj);
+					// Make sure we have a unique translation unit for each unique identity.
+					UTranslationUnit*& TranslationUnit = IdentityToTranslationUnitMap.FindOrAdd(FLocTextIdentity(ManifestEntry->Namespace, AContext.Key));
 					if (!TranslationUnit)
 					{
 						TranslationUnit = NewObject<UTranslationUnit>();
@@ -111,12 +156,13 @@ void FTranslationDataManager::Initialize()
 						TranslationUnit->HasBeenReviewed = false;
 						TranslationUnit->Source = ManifestEntry->Source.Text;
 						TranslationUnit->Namespace = ManifestEntry->Namespace;
+						TranslationUnit->Key = AContext.Key;
 						TranslationUnit->KeyMetaDataObject = AContext.KeyMetadataObj;
 					}
 					
 					if (NativeArchivePtr.IsValid() && NativeArchivePtr != ArchivePtr)
 					{
-						const TSharedPtr<FArchiveEntry> NativeArchiveEntry = NativeArchivePtr->FindEntryBySource(ManifestEntry->Namespace, ManifestEntry->Source, AContext.KeyMetadataObj);
+						const TSharedPtr<FArchiveEntry> NativeArchiveEntry = NativeArchivePtr->FindEntryByKey(ManifestEntry->Namespace, AContext.Key, AContext.KeyMetadataObj);
 						// If the native archive contains a translation for the source string that isn't identical to the source string, use the translation as the source string.
 						if (NativeArchiveEntry.IsValid() && !NativeArchiveEntry->Translation.IsExactMatch(NativeArchiveEntry->Source))
 						{
@@ -129,7 +175,7 @@ void FTranslationDataManager::Initialize()
 
 
 				TArray<UTranslationUnit*> TranslationUnitsToAdd;
-				KeyMetaDataToTranslationUnitMap.GenerateValueArray(TranslationUnitsToAdd);
+				IdentityToTranslationUnitMap.GenerateValueArray(TranslationUnitsToAdd);
 				TranslationUnits.Append(TranslationUnitsToAdd);
 			}
 			GWarn->EndSlowTask();
@@ -168,64 +214,30 @@ FTranslationDataManager::~FTranslationDataManager()
 
 TSharedPtr< FInternationalizationManifest > FTranslationDataManager::ReadManifest( const FString& ManifestFilePathToRead )
 {
-	
-	TSharedPtr<FJsonObject> ManifestJsonObject = ReadJSONTextFile( ManifestFilePathToRead );
+	TSharedPtr<FInternationalizationManifest> InternationalizationManifest = MakeShareable(new FInternationalizationManifest);
 
-	if( !ManifestJsonObject.IsValid() )
+	if (!FJsonInternationalizationManifestSerializer::DeserializeManifestFromFile(ManifestFilePathToRead, InternationalizationManifest.ToSharedRef()))
 	{
 		UE_LOG(LogTranslationEditor, Error, TEXT("Could not read manifest file %s."), *ManifestFilePathToRead);
-		return TSharedPtr< FInternationalizationManifest >();
+		InternationalizationManifest.Reset();
 	}
 
-	TSharedRef< FInternationalizationManifest > InternationalizationManifest = MakeShareable( new FInternationalizationManifest );
-
-	ManifestSerializer.DeserializeManifest( ManifestJsonObject.ToSharedRef(), InternationalizationManifest );
-
 	return InternationalizationManifest;
 }
 
 TSharedPtr< FInternationalizationArchive > FTranslationDataManager::ReadArchive(const FString& ArchiveFilePath)
 {
-	// Read in any existing archive for this culture.
-	TSharedPtr< FJsonObject > ArchiveJsonObject = ReadJSONTextFile( ArchiveFilePath );
+	TSharedPtr<FInternationalizationArchive> InternationalizationArchive = MakeShareable(new FInternationalizationArchive);
 
-	if ( !ArchiveJsonObject.IsValid() )
+	if (!FJsonInternationalizationArchiveSerializer::DeserializeArchiveFromFile(ArchiveFilePath, InternationalizationArchive.ToSharedRef(), ManifestAtHeadRevisionPtr, nullptr))
 	{
 		UE_LOG(LogTranslationEditor, Error, TEXT("Could not read archive file %s."), *ArchiveFilePath);
-		return nullptr;
+		InternationalizationArchive.Reset();
 	}
 
-	TSharedRef< FInternationalizationArchive > InternationalizationArchive = MakeShareable( new FInternationalizationArchive );
-
-	ArchiveSerializer.DeserializeArchive( ArchiveJsonObject.ToSharedRef(), InternationalizationArchive );
-
 	return InternationalizationArchive;
 }
 
-TSharedPtr<FJsonObject> FTranslationDataManager::ReadJSONTextFile(const FString& InFilePath)
-{
-	//read in file as string
-	FString FileContents;
-	if ( !FFileHelper::LoadFileToString(FileContents, *InFilePath) )
-	{
-		UE_LOG(LogTranslationEditor, Error,TEXT("Failed to load file %s."), *InFilePath);
-		return nullptr;
-	}
-
-	//parse as JSON
-	TSharedPtr<FJsonObject> JSONObject;
-
-	TSharedRef< TJsonReader<> > Reader = TJsonReaderFactory<>::Create( FileContents );
-
-	if( !FJsonSerializer::Deserialize( Reader, JSONObject ) || !JSONObject.IsValid())
-	{
-		UE_LOG(LogTranslationEditor, Error,TEXT("Invalid JSON in file %s."), *InFilePath);
-		return nullptr;
-	}
-
-	return JSONObject;
-}
-
 bool FTranslationDataManager::WriteTranslationData(bool bForceWrite /*= false*/)
 {
 	bool bSuccess = false;
@@ -242,11 +254,11 @@ bool FTranslationDataManager::WriteTranslationData(bool bForceWrite /*= false*/)
 			if (TranslationUnit != nullptr)
 			{
 				const FLocItem SearchSource(TranslationUnit->Source);
-				FString OldTranslation = Archive->FindEntryBySource(TranslationUnit->Namespace, SearchSource, nullptr)->Translation.Text;
+				FString OldTranslation = Archive->FindEntryByKey(TranslationUnit->Namespace, TranslationUnit->Key, TranslationUnit->KeyMetaDataObject)->Translation.Text;
 				FString TranslationToWrite = TranslationUnit->Translation;
 				if (!TranslationToWrite.Equals(OldTranslation))
 				{
-					Archive->SetTranslation(TranslationUnit->Namespace, TranslationUnit->Source, TranslationToWrite, nullptr);
+					Archive->SetTranslation(TranslationUnit->Namespace, TranslationUnit->Key, FLocItem(TranslationUnit->Source), FLocItem(TranslationToWrite), TranslationUnit->KeyMetaDataObject);
 					bNeedsWrite = true;
 				}
 			}
@@ -257,11 +269,11 @@ bool FTranslationDataManager::WriteTranslationData(bool bForceWrite /*= false*/)
 			if (TranslationUnit != nullptr)
 			{
 				const FLocItem SearchSource(TranslationUnit->Source);
-				FString OldTranslation = Archive->FindEntryBySource(TranslationUnit->Namespace, SearchSource, nullptr)->Translation.Text;
+				FString OldTranslation = Archive->FindEntryByKey(TranslationUnit->Namespace, TranslationUnit->Key, TranslationUnit->KeyMetaDataObject)->Translation.Text;
 				FString TranslationToWrite = TranslationUnit->Translation;
 				if (TranslationUnit->HasBeenReviewed && !TranslationToWrite.Equals(OldTranslation))
 				{
-					Archive->SetTranslation(TranslationUnit->Namespace, TranslationUnit->Source, TranslationToWrite, nullptr);
+					Archive->SetTranslation(TranslationUnit->Namespace, TranslationUnit->Key, FLocItem(TranslationUnit->Source), FLocItem(TranslationToWrite), TranslationUnit->KeyMetaDataObject);
 					bNeedsWrite = true;
 				}
 			}
@@ -272,11 +284,11 @@ bool FTranslationDataManager::WriteTranslationData(bool bForceWrite /*= false*/)
 			if (TranslationUnit != nullptr)
 			{
 				const FLocItem SearchSource(TranslationUnit->Source);
-				FString OldTranslation = Archive->FindEntryBySource(TranslationUnit->Namespace, SearchSource, nullptr)->Translation.Text;
+				FString OldTranslation = Archive->FindEntryByKey(TranslationUnit->Namespace, TranslationUnit->Key, TranslationUnit->KeyMetaDataObject)->Translation.Text;
 				FString TranslationToWrite = TranslationUnit->Translation;
 				if (!TranslationToWrite.Equals(OldTranslation))
 				{
-					Archive->SetTranslation(TranslationUnit->Namespace, TranslationUnit->Source, TranslationToWrite, nullptr);
+					Archive->SetTranslation(TranslationUnit->Namespace, TranslationUnit->Key, FLocItem(TranslationUnit->Source), FLocItem(TranslationToWrite), TranslationUnit->KeyMetaDataObject);
 					bNeedsWrite = true;
 				}
 			}
@@ -287,7 +299,7 @@ bool FTranslationDataManager::WriteTranslationData(bool bForceWrite /*= false*/)
 		if (bForceWrite || bNeedsWrite)
 		{
 			TSharedRef<FJsonObject> FinalArchiveJsonObj = MakeShareable(new FJsonObject);
-			ArchiveSerializer.SerializeArchive(Archive, FinalArchiveJsonObj);
+			FJsonInternationalizationArchiveSerializer::SerializeArchive(Archive, FinalArchiveJsonObj);
 
 			bSuccess = WriteJSONToTextFile(FinalArchiveJsonObj, OpenedArchiveFilePath);
 		}
@@ -498,9 +510,7 @@ void FTranslationDataManager::GetHistoryForTranslationUnits()
 									PreviousSourceText = ContextInfo.Changes[0].Source;
 								}
 
-								FManifestContext SearchContext;
-								SearchContext.Key = ContextInfo.Key;
-								TSharedPtr< FManifestEntry > OldManifestEntryPtr = OldManifest->FindEntryByContext(TranslationUnit->Namespace, SearchContext);
+								TSharedPtr< FManifestEntry > OldManifestEntryPtr = OldManifest->FindEntryByKey(TranslationUnit->Namespace, ContextInfo.Key);
 								if (!OldManifestEntryPtr.IsValid())
 								{
 									// If this version of the manifest didn't know anything about this string, move onto the next
@@ -510,7 +520,7 @@ void FTranslationDataManager::GetHistoryForTranslationUnits()
 								// Always add first instance of this string, and then add any versions that changed since
 								if (ContextInfo.Changes.Num() == 0 || !OldManifestEntryPtr->Source.Text.Equals(PreviousSourceText))
 								{
-									TSharedPtr< FArchiveEntry > OldArchiveEntry = ArchivePtr->FindEntryBySource(OldManifestEntryPtr->Namespace, OldManifestEntryPtr->Source, nullptr);
+									TSharedPtr< FArchiveEntry > OldArchiveEntry = ArchivePtr->FindEntryByKey(OldManifestEntryPtr->Namespace, ContextInfo.Key, nullptr);
 									if (OldArchiveEntry.IsValid())
 									{
 										FTranslationChange Change;
@@ -625,11 +635,7 @@ void FTranslationDataManager::PreviewAllTranslationsInEditor(ULocalizationTarget
 	{
 		const FString ConfigFilePath = LocalizationConfigurationScript::GetRegenerateResourcesConfigPath(LocalizationTarget);
 		LocalizationConfigurationScript::GenerateRegenerateResourcesConfigFile(LocalizationTarget).Write(ConfigFilePath);
-
-		FJsonInternationalizationArchiveSerializer LocalizationArchiveSerializer;
-		FJsonInternationalizationManifestSerializer LocalizationManifestSerializer;
-
-		FTextLocalizationManager::Get().LoadFromManifestAndArchives(ConfigFilePath, LocalizationArchiveSerializer, LocalizationManifestSerializer);
+		FTextLocalizationResourceGenerator::GenerateAndUpdateLiveEntriesFromConfig(ConfigFilePath, /*bSkipSourceCheck*/false);
 	}
 	else
 	{
@@ -708,8 +714,7 @@ void FTranslationDataManager::LoadFromArchive(TArray<UTranslationUnit*>& InTrans
 
 				GWarn->StatusUpdate(CurrentTranslationUnitIndex, TranslationUnits.Num(), FText::Format(LOCTEXT("LoadingCurrentArchiveEntries", "Loading Entry {0} of {1} from Translation Archive..."), FText::AsNumber(CurrentTranslationUnitIndex), FText::AsNumber(TranslationUnits.Num())));
 
-				const FLocItem SourceSearch(TranslationUnit->Source);
-				TSharedPtr<FArchiveEntry> ArchiveEntry = Archive->FindEntryBySource(TranslationUnit->Namespace, SourceSearch, TranslationUnit->KeyMetaDataObject);
+				TSharedPtr<FArchiveEntry> ArchiveEntry = Archive->FindEntryByKey(TranslationUnit->Namespace, TranslationUnit->Key, TranslationUnit->KeyMetaDataObject);
 				if (ArchiveEntry.IsValid())
 				{
 					const FString PreviousTranslation = TranslationUnit->Translation;
diff --git a/Engine/Source/Editor/TranslationEditor/Private/TranslationDataManager.h b/Engine/Source/Editor/TranslationEditor/Private/TranslationDataManager.h
--- a/Engine/Source/Editor/TranslationEditor/Private/TranslationDataManager.h
+++ b/Engine/Source/Editor/TranslationEditor/Private/TranslationDataManager.h
@@ -7,8 +7,6 @@ class FInternationalizationArchive;
 class FInternationalizationManifest;
 class UTranslationUnit;
 
-#include "JsonInternationalizationArchiveSerializer.h"
-#include "JsonInternationalizationManifestSerializer.h"
 #include "ILocalizationServiceProvider.h"
 
 class FTranslationDataManager : public TSharedFromThis<FTranslationDataManager>
@@ -90,9 +88,6 @@ public:
 private:
 	void Initialize();
 
-	/** Read text file into a JSON file */
-	static TSharedPtr<FJsonObject> ReadJSONTextFile( const FString& InFilePath ) ;
-
 	/** Take a path and a manifest name and return a manifest data structure */
 	TSharedPtr< FInternationalizationManifest > ReadManifest ( const FString& ManifestFilePath );
 	/** Retrieve an archive data structure from ArchiveFilePath */
@@ -112,11 +107,6 @@ private:
 	TArray<UTranslationUnit*> SearchResults;
 	TArray<UTranslationUnit*> ChangedOnImport;
 
-	/** Serializes and deserializes our Archive */
-	FJsonInternationalizationArchiveSerializer ArchiveSerializer;
-	/** Serializes and deserializes Manifests */
-	FJsonInternationalizationManifestSerializer ManifestSerializer;
-
 	/** Archive for the current project and native language */
 	TSharedPtr< FInternationalizationArchive > NativeArchivePtr;
 	/** Archive for the current project and translation language */
diff --git a/Engine/Source/Editor/UMGEditor/Private/Hierarchy/SHierarchyView.cpp b/Engine/Source/Editor/UMGEditor/Private/Hierarchy/SHierarchyView.cpp
--- a/Engine/Source/Editor/UMGEditor/Private/Hierarchy/SHierarchyView.cpp
+++ b/Engine/Source/Editor/UMGEditor/Private/Hierarchy/SHierarchyView.cpp
@@ -296,6 +296,7 @@ void SHierarchyView::RebuildTreeView()
 		.OnSelectionChanged(this, &SHierarchyView::WidgetHierarchy_OnSelectionChanged)
 		.OnExpansionChanged(this, &SHierarchyView::WidgetHierarchy_OnExpansionChanged)
 		.OnContextMenuOpening(this, &SHierarchyView::WidgetHierarchy_OnContextMenuOpening)
+		.OnSetExpansionRecursive(this, &SHierarchyView::SetItemExpansionRecursive)
 		.TreeItemsSource(&TreeRootWidgets);
 
 	FilterHandler->SetTreeView(WidgetTreeView.Get());
@@ -324,19 +325,16 @@ void SHierarchyView::RestoreExpandedItems()
 	}
 }
 
-void SHierarchyView::RecursiveExpand(TSharedPtr<FHierarchyModel>& Model)
+void SHierarchyView::RecursiveExpand(TSharedPtr<FHierarchyModel>& Model, bool InShouldExpandItem)
 {
-	if ( Model->IsExpanded() )
-	{
-		WidgetTreeView->SetItemExpansion(Model, true);
+	WidgetTreeView->SetItemExpansion(Model, InShouldExpandItem);
 
-		TArray< TSharedPtr<FHierarchyModel> > Children;
-		Model->GatherChildren(Children);
+	TArray< TSharedPtr<FHierarchyModel> > Children;
+	Model->GatherChildren(Children);
 
-		for ( TSharedPtr<FHierarchyModel>& ChildModel : Children )
-		{
-			RecursiveExpand(ChildModel);
-		}
+	for (TSharedPtr<FHierarchyModel>& ChildModel : Children)
+	{
+		RecursiveExpand(ChildModel, InShouldExpandItem);
 	}
 }
 
@@ -371,6 +369,13 @@ void SHierarchyView::RecursiveSelection(TSharedPtr<FHierarchyModel>& Model)
 	}
 }
 
+void SHierarchyView::SetItemExpansionRecursive(TSharedPtr<FHierarchyModel> Model, bool bInExpansionState)
+{
+	if (Model.IsValid())
+	{
+		RecursiveExpand(Model, bInExpansionState);
+	}
+}
 
 //@TODO UMG Drop widgets onto the tree, when nothing is present, if there is a root node present, what happens then, let the root node attempt to place it?
 
diff --git a/Engine/Source/Editor/UMGEditor/Private/Hierarchy/SHierarchyView.h b/Engine/Source/Editor/UMGEditor/Private/Hierarchy/SHierarchyView.h
--- a/Engine/Source/Editor/UMGEditor/Private/Hierarchy/SHierarchyView.h
+++ b/Engine/Source/Editor/UMGEditor/Private/Hierarchy/SHierarchyView.h
@@ -92,7 +92,7 @@ private:
 	void RestoreExpandedItems();
 
 	/** Recursively expands the models based on the expansion set. */
-	void RecursiveExpand(TSharedPtr<FHierarchyModel>& Model);
+	void RecursiveExpand(TSharedPtr<FHierarchyModel>& Model, bool InShouldExpandItem = true);
 
 	/**  */
 	void RestoreSelectedItems();
@@ -100,6 +100,9 @@ private:
 	/**  */
 	void RecursiveSelection(TSharedPtr<FHierarchyModel>& Model);
 
+	/** Handler for recursively expanding/collapsing items */
+	void SetItemExpansionRecursive(TSharedPtr<FHierarchyModel> Model, bool bInExpansionState);
+
 private:
 
 	/** Cached pointer to the blueprint editor that owns this tree. */
diff --git a/Engine/Source/Editor/UMGEditor/Private/Palette/SPaletteView.cpp b/Engine/Source/Editor/UMGEditor/Private/Palette/SPaletteView.cpp
--- a/Engine/Source/Editor/UMGEditor/Private/Palette/SPaletteView.cpp
+++ b/Engine/Source/Editor/UMGEditor/Private/Palette/SPaletteView.cpp
@@ -8,8 +8,6 @@
 #include "PreviewScene.h"
 #include "SceneViewport.h"
 
-#include "BlueprintEditor.h"
-
 #include "DecoratedDragDropOp.h"
 #include "WidgetTemplateDragDropOp.h"
 
@@ -22,6 +20,7 @@
 #include "AssetRegistryModule.h"
 #include "SSearchBox.h"
 
+#include "WidgetBlueprintEditor.h"
 #include "WidgetBlueprintCompiler.h"
 #include "WidgetBlueprintEditorUtils.h"
 
@@ -32,6 +31,7 @@
 #include "ObjectEditorUtils.h"
 
 #include "Settings/ContentBrowserSettings.h"
+#include "UMGEditorProjectSettings.h"
 
 #define LOCTEXT_NAMESPACE "UMG"
 
@@ -103,6 +103,11 @@ public:
 		return Template->Name;
 	}
 
+	virtual bool IsTemplate() const override
+	{
+		return true;
+	}
+
 	virtual FString GetFilterString() const override
 	{
 		return Template->Name.ToString();
@@ -141,6 +146,11 @@ public:
 		return GroupName;
 	}
 
+	virtual bool IsTemplate() const override
+	{
+		return false;
+	}
+
 	virtual FString GetFilterString() const override
 	{
 		// Headers should never be included in filtering to avoid showing a header with all of
@@ -174,7 +184,7 @@ public:
 	TArray< TSharedPtr<FWidgetViewModel> > Children;
 };
 
-void SPaletteView::Construct(const FArguments& InArgs, TSharedPtr<FBlueprintEditor> InBlueprintEditor)
+void SPaletteView::Construct(const FArguments& InArgs, TSharedPtr<FWidgetBlueprintEditor> InBlueprintEditor)
 {
 	// Register for events that can trigger a palette rebuild
 	GEditor->OnBlueprintReinstanced().AddRaw(this, &SPaletteView::OnBlueprintReinstanced);
@@ -201,6 +211,7 @@ void SPaletteView::Construct(const FArguments& InArgs, TSharedPtr<FBlueprintEdit
 		.SelectionMode(ESelectionMode::Single)
 		.OnGenerateRow(this, &SPaletteView::OnGenerateWidgetTemplateItem)
 		.OnGetChildren(FilterHandler.ToSharedRef(), &PaletteFilterHandler::OnGetFilteredChildren)
+		.OnSelectionChanged(this, &SPaletteView::WidgetPalette_OnSelectionChanged)
 		.TreeItemsSource(&TreeWidgetViewModels);
 		
 
@@ -270,6 +281,61 @@ FText SPaletteView::GetSearchText() const
 	return SearchText;
 }
 
+void SPaletteView::WidgetPalette_OnSelectionChanged(TSharedPtr<FWidgetViewModel> SelectedItem, ESelectInfo::Type SelectInfo)
+{
+	if (!SelectedItem.IsValid()) 
+	{
+		return;
+	}
+
+	// Reset the selected
+	BlueprintEditor.Pin()->SetSelectedTemplate(nullptr);
+	BlueprintEditor.Pin()->SetSelectedUserWidget(FAssetData());
+
+	// If it's not a template, return
+	if (!SelectedItem->IsTemplate())
+	{
+		return;
+	}
+
+	TSharedPtr<FWidgetTemplateViewModel> SelectedTemplate = StaticCastSharedPtr<FWidgetTemplateViewModel>(SelectedItem);
+	if (SelectedTemplate.IsValid())
+	{
+		TSharedPtr<FWidgetTemplateClass> TemplateClass = StaticCastSharedPtr<FWidgetTemplateClass>(SelectedTemplate->Template);
+		if (TemplateClass.IsValid())
+		{
+			if (TemplateClass->GetWidgetClass().IsValid())
+			{
+				BlueprintEditor.Pin()->SetSelectedTemplate(TemplateClass->GetWidgetClass());
+			}
+			else
+			{
+				TSharedPtr<FWidgetTemplateBlueprintClass> UserCreatedTemplate = StaticCastSharedPtr<FWidgetTemplateBlueprintClass>(TemplateClass);
+				if (UserCreatedTemplate.IsValid())
+				{
+					// Then pass in the asset data of selected widget
+					FAssetData UserCreatedWidget = UserCreatedTemplate->GetWidgetAssetData();
+					BlueprintEditor.Pin()->SetSelectedUserWidget(UserCreatedWidget);
+				}
+			}
+		}
+	}
+}
+
+TSharedPtr<FWidgetTemplate> SPaletteView::GetSelectedTemplateWidget() const
+{
+	TArray<TSharedPtr<FWidgetViewModel>> SelectedTemplates = WidgetTemplatesView.Get()->GetSelectedItems();
+	if (SelectedTemplates.Num() == 1)
+	{
+		TSharedPtr<FWidgetTemplateViewModel> TemplateViewModel = StaticCastSharedPtr<FWidgetTemplateViewModel>(SelectedTemplates[0]);
+		if (TemplateViewModel.IsValid())
+		{
+			return TemplateViewModel->Template;
+		}
+	}
+	return nullptr;
+}
+
 void SPaletteView::LoadItemExpansion()
 {
 	// Restore the expansion state of the widget groups.
@@ -345,7 +411,8 @@ void SPaletteView::BuildClassWidgetList()
 
 	auto ActiveWidgetBlueprintClass = GetBlueprint()->GeneratedClass;
 	FName ActiveWidgetBlueprintClassName = ActiveWidgetBlueprintClass->GetFName();
-	UWidgetDesignerSettings* ViewportSettings = GetMutableDefault<UWidgetDesignerSettings>();
+
+	TArray<FStringClassReference> WidgetClassesToHide = GetDefault<UUMGEditorProjectSettings>()->WidgetClassesToHide;
 
 	// Locate all UWidget classes from code and loaded widget BPs
 	for (TObjectIterator<UClass> ClassIt; ClassIt; ++ClassIt)
@@ -356,23 +423,38 @@ void SPaletteView::BuildClassWidgetList()
 		FAssetData WidgetAssetData = FAssetData(WidgetClass);
 
 		// Excludes engine content if user sets it to false
-		if (!GetDefault<UContentBrowserSettings>()->GetDisplayEngineFolder())
+		if (!GetDefault<UContentBrowserSettings>()->GetDisplayEngineFolder() || !GetDefault<UUMGEditorProjectSettings>()->bShowWidgetsFromEngineContent)
 		{
-			if (WidgetAssetData.PackagePath.ToString().Find((TEXT("/Engine"))) == 0)
+			if (WidgetAssetData.PackagePath.ToString().Find(TEXT("/Engine")) == 0)
 			{
 				continue;
 			}
 		}
 
 		// Excludes developer content if user sets it to false
-		if (!GetDefault<UContentBrowserSettings>()->GetDisplayDevelopersFolder())
+		if (!GetDefault<UContentBrowserSettings>()->GetDisplayDevelopersFolder() || !GetDefault<UUMGEditorProjectSettings>()->bShowWidgetsFromDeveloperContent)
 		{
-			if (WidgetAssetData.PackagePath.ToString().Find((TEXT("/Game/Developers"))) == 0)
+			if (WidgetAssetData.PackagePath.ToString().Find(TEXT("/Game/Developers")) == 0)
 			{
 				continue;
 			}
 		}
 
+		// Excludes this widget if it is on the hide list
+		bool bIsOnList = false;
+		for (FStringClassReference Widget : WidgetClassesToHide)
+		{
+			if (WidgetAssetData.ObjectPath.ToString().Find(Widget.ToString()) == 0)
+			{
+				bIsOnList = true;
+				break;
+			}
+		}
+		if (bIsOnList)
+		{
+			continue;
+		}
+
 		if ( FWidgetBlueprintEditorUtils::IsUsableWidgetClass(WidgetClass) )
 		{
 			const bool bIsSameClass = WidgetClass->GetFName() == ActiveWidgetBlueprintClassName;
@@ -417,23 +499,38 @@ void SPaletteView::BuildClassWidgetList()
 		}
 
 		// Excludes engine content if user sets it to false
-		if (!GetDefault<UContentBrowserSettings>()->GetDisplayEngineFolder() )
+		if (!GetDefault<UContentBrowserSettings>()->GetDisplayEngineFolder() || !GetDefault<UUMGEditorProjectSettings>()->bShowWidgetsFromEngineContent)
 		{
-			if (WidgetBPAssetData.PackagePath.ToString().Find((TEXT("/Engine"))) == 0)
+			if (WidgetBPAssetData.PackagePath.ToString().Find(TEXT("/Engine")) == 0)
 			{
 				continue;
 			}
 		}
 
 		// Excludes developer content if user sets it to false
-		if (!GetDefault<UContentBrowserSettings>()->GetDisplayDevelopersFolder() )
+		if (!GetDefault<UContentBrowserSettings>()->GetDisplayDevelopersFolder() || !GetDefault<UUMGEditorProjectSettings>()->bShowWidgetsFromDeveloperContent)
 		{
-			if (WidgetBPAssetData.PackagePath.ToString().Find((TEXT("/Game/Developers"))) == 0)
+			if (WidgetBPAssetData.PackagePath.ToString().Find(TEXT("/Game/Developers")) == 0)
 			{
 				continue;
 			}
 		}
 
+		// Excludes this widget if it is on the hide list
+		bool bIsOnList = false;
+		for (FStringClassReference Widget : WidgetClassesToHide)
+		{
+			if (Widget.ToString().Find(WidgetBPAssetData.ObjectPath.ToString()) == 0)
+			{
+				bIsOnList = true;
+				break;
+			}
+		}
+		if (bIsOnList)
+		{
+			continue;
+		}
+
 		// If the blueprint generated class was found earlier, pass it to the template
 		TSubclassOf<UUserWidget> WidgetBPClass = nullptr;
 		auto LoadedWidgetBPClass = LoadedWidgetBlueprintClassesByName.Find(WidgetBPAssetData.AssetName);
@@ -459,6 +556,16 @@ void SPaletteView::BuildSpecialWidgetList()
 void SPaletteView::AddWidgetTemplate(TSharedPtr<FWidgetTemplate> Template)
 {
 	FString Category = Template->GetCategory().ToString();
+
+	// Hide user specific categories
+	TArray<FString> CategoriesToHide = GetDefault<UUMGEditorProjectSettings>()->CategoriesToHide;
+	for (FString CategoryName : CategoriesToHide)
+	{
+		if (Category == CategoryName)
+		{
+			return;
+		}
+	}
 	WidgetTemplateArray& Group = WidgetTemplateCategories.FindOrAdd(Category);
 	Group.Add(Template);
 }
diff --git a/Engine/Source/Editor/UMGEditor/Private/Palette/SPaletteView.h b/Engine/Source/Editor/UMGEditor/Private/Palette/SPaletteView.h
--- a/Engine/Source/Editor/UMGEditor/Private/Palette/SPaletteView.h
+++ b/Engine/Source/Editor/UMGEditor/Private/Palette/SPaletteView.h
@@ -3,7 +3,7 @@
 #pragma once
 #include "Misc/TextFilter.h"
 #include "SCompoundWidget.h"
-#include "BlueprintEditor.h"
+#include "WidgetBlueprintEditor.h"
 #include "TreeFilterHandler.h"
 
 class FWidgetTemplate;
@@ -15,6 +15,8 @@ class FWidgetViewModel : public TSharedFromThis<FWidgetViewModel>
 public:
 	virtual FText GetName() const = 0;
 
+	virtual bool IsTemplate() const = 0;
+
 	/** Get the string which should be used for filtering the item. */
 	virtual FString GetFilterString() const = 0;
 
@@ -35,7 +37,7 @@ public:
 	SLATE_BEGIN_ARGS( SPaletteView ){}
 	SLATE_END_ARGS()
 
-	void Construct(const FArguments& InArgs, TSharedPtr<FBlueprintEditor> InBlueprintEditor);
+	void Construct(const FArguments& InArgs, TSharedPtr<FWidgetBlueprintEditor> InBlueprintEditor);
 	virtual ~SPaletteView();
 	
 	virtual void Tick( const FGeometry& AllottedGeometry, const double InCurrentTime, const float InDeltaTime ) override;
@@ -43,6 +45,12 @@ public:
 	/** Gets the text currently displayed in the search box. */
 	FText GetSearchText() const;
 
+	/** On Selected template widget changed in palette view */
+	void WidgetPalette_OnSelectionChanged(TSharedPtr<FWidgetViewModel> SelectedItem, ESelectInfo::Type SelectInfo);
+
+	/** Gets the selected template widget in palette */
+	TSharedPtr<FWidgetTemplate> GetSelectedTemplateWidget() const;
+
 private:
 	UWidgetBlueprint* GetBlueprint() const;
 
@@ -77,7 +85,7 @@ private:
 	/** Requests a rebuild of the widget list if a widget blueprint was deleted */
 	void HandleOnAssetsDeleted(const TArray<UClass*>& DeletedAssetClasses);
 
-	TWeakPtr<class FBlueprintEditor> BlueprintEditor;
+	TWeakPtr<class FWidgetBlueprintEditor> BlueprintEditor;
 
 	/** Handles filtering the palette based on an IFilter. */
 	typedef TreeFilterHandler<TSharedPtr<FWidgetViewModel>> PaletteFilterHandler;
diff --git a/Engine/Source/Editor/UMGEditor/Private/Settings/UMGEditorProjectSettings.cpp b/Engine/Source/Editor/UMGEditor/Private/Settings/UMGEditorProjectSettings.cpp
new file mode 100644
--- /dev/null
+++ b/Engine/Source/Editor/UMGEditor/Private/Settings/UMGEditorProjectSettings.cpp
@@ -0,0 +1,11 @@
+// Copyright 1998-2016 Epic Games, Inc. All Rights Reserved.
+
+#include "UMGEditorPrivatePCH.h"
+#include "UMGEditorProjectSettings.h"
+
+UUMGEditorProjectSettings::UUMGEditorProjectSettings(const FObjectInitializer& ObjectInitializer)
+	:Super(ObjectInitializer)
+{
+	bShowWidgetsFromEngineContent = false;
+	bShowWidgetsFromDeveloperContent = true;
+}
\ No newline at end of file
diff --git a/Engine/Source/Editor/UMGEditor/Private/TabFactory/AnimationTabSummoner.cpp b/Engine/Source/Editor/UMGEditor/Private/TabFactory/AnimationTabSummoner.cpp
--- a/Engine/Source/Editor/UMGEditor/Private/TabFactory/AnimationTabSummoner.cpp
+++ b/Engine/Source/Editor/UMGEditor/Private/TabFactory/AnimationTabSummoner.cpp
@@ -50,6 +50,13 @@ bool VerifyAnimationRename( UWidgetBlueprint* Blueprint, UWidgetAnimation* Anima
 		return false;
 	}
 
+	FName FunctionName(*NewAnimationName);
+	if (Animation->GetPreviewWidget().IsValid() && Animation->GetPreviewWidget().Get()->FindFunction(FunctionName))
+	{
+		OutErrorMessage = LOCTEXT("NameInUseByFunction", "A function with this name already exists");
+		return false;
+	}
+
 	FKismetNameValidator Validator( Blueprint );
 	EValidatorResult ValidationResult = Validator.IsValid( NewAnimationName );
 
diff --git a/Engine/Source/Editor/UMGEditor/Private/Templates/WidgetTemplateBlueprintClass.cpp b/Engine/Source/Editor/UMGEditor/Private/Templates/WidgetTemplateBlueprintClass.cpp
--- a/Engine/Source/Editor/UMGEditor/Private/Templates/WidgetTemplateBlueprintClass.cpp
+++ b/Engine/Source/Editor/UMGEditor/Private/Templates/WidgetTemplateBlueprintClass.cpp
@@ -97,4 +97,9 @@ FReply FWidgetTemplateBlueprintClass::OnDoubleClicked()
 	return FReply::Handled();
 }
 
+FAssetData FWidgetTemplateBlueprintClass::GetWidgetAssetData()
+{
+	return WidgetAssetData;
+}
+
 #undef LOCTEXT_NAMESPACE
diff --git a/Engine/Source/Editor/UMGEditor/Private/Templates/WidgetTemplateBlueprintClass.h b/Engine/Source/Editor/UMGEditor/Private/Templates/WidgetTemplateBlueprintClass.h
--- a/Engine/Source/Editor/UMGEditor/Private/Templates/WidgetTemplateBlueprintClass.h
+++ b/Engine/Source/Editor/UMGEditor/Private/Templates/WidgetTemplateBlueprintClass.h
@@ -39,6 +39,9 @@ public:
 	/** Opens the widget blueprint for edit */
 	virtual FReply OnDoubleClicked() override;
 
+	/** Gets the asset data for this widget blueprint */
+	FAssetData GetWidgetAssetData();
+
 protected:
 	/** The asset data for the widget blueprint */
 	FAssetData WidgetAssetData;
diff --git a/Engine/Source/Editor/UMGEditor/Private/Templates/WidgetTemplateClass.h b/Engine/Source/Editor/UMGEditor/Private/Templates/WidgetTemplateClass.h
--- a/Engine/Source/Editor/UMGEditor/Private/Templates/WidgetTemplateClass.h
+++ b/Engine/Source/Editor/UMGEditor/Private/Templates/WidgetTemplateClass.h
@@ -30,6 +30,9 @@ public:
 	/** Gets the tooltip widget for this palette item. */
 	virtual TSharedRef<IToolTip> GetToolTip() const override;
 
+	/** Gets the WidgetClass */
+	TWeakObjectPtr<UClass> GetWidgetClass() const { return WidgetClass; }
+
 protected:
 	/** Creates a widget template class without any class reference */
 	FWidgetTemplateClass();
diff --git a/Engine/Source/Editor/UMGEditor/Private/UMGEditorModule.cpp b/Engine/Source/Editor/UMGEditor/Private/UMGEditorModule.cpp
--- a/Engine/Source/Editor/UMGEditor/Private/UMGEditorModule.cpp
+++ b/Engine/Source/Editor/UMGEditor/Private/UMGEditorModule.cpp
@@ -24,6 +24,7 @@
 #include "ClassIconFinder.h"
 
 #include "Settings/WidgetDesignerSettings.h"
+#include "UMGEditorProjectSettings.h"
 #include "ISettingsModule.h"
 
 #define LOCTEXT_NAMESPACE "UMG"
@@ -84,6 +85,13 @@ public:
 				LOCTEXT("WidgetDesignerSettingsDescription", "Configure options for the Widget Designer."),
 				GetMutableDefault<UWidgetDesignerSettings>()
 				);
+
+			// UMG Editor Settings
+			SettingsModule->RegisterSettings("Project", "Editor", "UMGEditor",
+				LOCTEXT("UMGEditorSettingsName", "UMG Editor"),
+				LOCTEXT("UMGEditorSettingsDescription", "Configure options for UMG Editor."),
+				GetMutableDefault<UUMGEditorProjectSettings>()
+				);
 		}
 	}
 
@@ -119,6 +127,7 @@ public:
 		if ( SettingsModule != nullptr )
 		{
 			SettingsModule->UnregisterSettings("Editor", "ContentEditors", "WidgetDesigner");
+			SettingsModule->UnregisterSettings("Project", "Editor", "UMGEditor");
 		}
 	}
 
diff --git a/Engine/Source/Editor/UMGEditor/Private/WidgetBlueprintEditor.cpp b/Engine/Source/Editor/UMGEditor/Private/WidgetBlueprintEditor.cpp
--- a/Engine/Source/Editor/UMGEditor/Private/WidgetBlueprintEditor.cpp
+++ b/Engine/Source/Editor/UMGEditor/Private/WidgetBlueprintEditor.cpp
@@ -100,6 +100,9 @@ void FWidgetBlueprintEditor::InitWidgetBlueprintEditor(const EToolkitMode::Type
 	// register for any objects replaced
 	GEditor->OnObjectsReplaced().AddSP(this, &FWidgetBlueprintEditor::OnObjectsReplaced);
 
+	// for change selected widgets on sequencer tree view
+	Sequencer->GetSelectionChangedObjectGuids().AddSP(this, &FWidgetBlueprintEditor::SyncSelectedWidgetsWithSequencerSelection);
+
 	UWidgetBlueprint* Blueprint = GetWidgetBlueprintObj();
 
 	// If this blueprint is empty, add a canvas panel as the root widget.
@@ -1118,7 +1121,9 @@ void FWidgetBlueprintEditor::ReplaceTrackWithSelectedWidget(FWidgetReference Sel
 					UMovieScenePropertyTrack* PropertyTrack = Cast<UMovieScenePropertyTrack>(Track);
 					if (PropertyTrack)
 					{
-						FString NameString = "Set" + PropertyTrack->GetPropertyName().ToString();
+						FString PropertyName = PropertyTrack->GetPropertyName().ToString();
+						PropertyName.RemoveFromStart("b", ESearchCase::CaseSensitive);
+						FString NameString = "Set" + PropertyName;
 						FName FunctionName = FName(*NameString);
 						if (!SelectedWidget.GetTemplate()->FindFunction(FunctionName))
 						{
@@ -1206,7 +1211,7 @@ void FWidgetBlueprintEditor::ExtendSequencerObjectBindingMenu(FMenuBuilder& Obje
 			FUIAction ReplaceWithMenuAction(FExecuteAction::CreateRaw(this, &FWidgetBlueprintEditor::ReplaceTrackWithSelectedWidget, SelectedWidget, BoundWidget));
 
 			FText ReplaceWithLabel = FText::Format(LOCTEXT("ReplaceObject", "Replace with {0}"), FText::FromString(SelectedWidget.GetPreview()->GetName()));
-			FText ReplaceWithToolTip = FText::Format(LOCTEXT("ReplaceObjectToolTip", "Replace the widget in this animation with selected"), FText::FromString(SelectedWidget.GetPreview()->GetName()));
+			FText ReplaceWithToolTip = FText::Format(LOCTEXT("ReplaceObjectToolTip", "Replace the bound widget in this animation with {0}"), FText::FromString(SelectedWidget.GetPreview()->GetName()));
 
 			ObjectBindingMenuBuilder.AddMenuEntry(ReplaceWithLabel, ReplaceWithToolTip, FSlateIcon(), ReplaceWithMenuAction);
 			ObjectBindingMenuBuilder.AddMenuSeparator();
@@ -1251,4 +1256,32 @@ void FWidgetBlueprintEditor::OnMovieSceneDataChanged()
 	bRefreshGeneratedClassAnimations = true;
 }
 
+void FWidgetBlueprintEditor::SyncSelectedWidgetsWithSequencerSelection(TArray<FGuid> ObjectGuids)
+{
+	UWidgetAnimation* WidgetAnimation = Cast<UWidgetAnimation>(GetSequencer().Get()->GetFocusedMovieSceneSequence());
+	UObject* BindingContext = GetSequencer().Get()->GetPlaybackContext();
+	TSet<FWidgetReference> SequencerSelectedWidgets;
+	for (FGuid Guid : ObjectGuids)
+	{
+		UObject* BoundObject = WidgetAnimation->FindPossessableObject(Guid, BindingContext);
+		if (!BoundObject)
+		{
+			continue;
+		}
+		else if (Cast<UPanelSlot>(BoundObject))
+		{
+			SequencerSelectedWidgets.Add(GetReferenceFromPreview(Cast<UPanelSlot>(BoundObject)->Content));
+		}
+		else
+		{
+			UWidget* BoundWidget = Cast<UWidget>(BoundObject);
+			SequencerSelectedWidgets.Add(GetReferenceFromPreview(BoundWidget));
+		}
+	}
+	if (SequencerSelectedWidgets.Num() != 0)
+	{
+		SelectWidgets(SequencerSelectedWidgets, false);
+	}
+}
+
 #undef LOCTEXT_NAMESPACE
diff --git a/Engine/Source/Editor/UMGEditor/Private/WidgetBlueprintEditorUtils.cpp b/Engine/Source/Editor/UMGEditor/Private/WidgetBlueprintEditorUtils.cpp
--- a/Engine/Source/Editor/UMGEditor/Private/WidgetBlueprintEditorUtils.cpp
+++ b/Engine/Source/Editor/UMGEditor/Private/WidgetBlueprintEditorUtils.cpp
@@ -10,6 +10,7 @@
 #include "TextPackageNamespaceUtil.h"
 #include "K2Node_Variable.h"
 #include "WidgetTemplateClass.h"
+#include "WidgetTemplateBlueprintClass.h"
 #include "Factories.h"
 #include "UnrealExporter.h"
 #include "GenericCommands.h"
@@ -229,20 +230,22 @@ bool FWidgetBlueprintEditorUtils::RenameWidget(TSharedRef<FWidgetBlueprintEditor
 
 					WidgetAnimation->MovieScene->Modify();
 
-					FMovieScenePossessable* Possessable = WidgetAnimation->MovieScene->FindPossessable(AnimBinding.AnimationGuid);
-					if ( Possessable )
+					if (AnimBinding.SlotWidgetName == NAME_None)
 					{
-						Possessable->SetName(NewFName.ToString());
+						FMovieScenePossessable* Possessable = WidgetAnimation->MovieScene->FindPossessable(AnimBinding.AnimationGuid);
+						if (Possessable)
+						{
+							Possessable->SetName(NewFName.ToString());
+						}
+					}
+					else
+					{
+						break;
 					}
-
-					break;
 				}
 			}
 		}
 
-		// Update named slot bindings?
-		// TODO...
-
 		// Validate child blueprints and adjust variable names to avoid a potential name collision
 		FBlueprintEditorUtils::ValidateBlueprintChildVariables(Blueprint, NewFName);
 
@@ -296,8 +299,8 @@ void FWidgetBlueprintEditorUtils::CreateWidgetContextMenu(FMenuBuilder& MenuBuil
 		{
 			MenuBuilder.AddSubMenu(
 				LOCTEXT("WidgetTree_ReplaceWith", "Replace With..."),
-				LOCTEXT("WidgetTree_ReplaceWithToolTip", "Replaces the currently selected panel widget, with another panel widget"),
-				FNewMenuDelegate::CreateStatic(&FWidgetBlueprintEditorUtils::BuildReplaceWithMenu, BP, Widgets)
+				LOCTEXT("WidgetTree_ReplaceWithToolTip", "Replaces the currently selected widget, with another widget"),
+				FNewMenuDelegate::CreateStatic(&FWidgetBlueprintEditorUtils::BuildReplaceWithMenu, BlueprintEditor, BP, Widgets)
 				);
 		}
 	}
@@ -536,7 +539,7 @@ void FWidgetBlueprintEditorUtils::WrapWidgets(TSharedRef<FWidgetBlueprintEditor>
 	FBlueprintEditorUtils::MarkBlueprintAsStructurallyModified(BP);
 }
 
-void FWidgetBlueprintEditorUtils::BuildReplaceWithMenu(FMenuBuilder& Menu, UWidgetBlueprint* BP, TSet<FWidgetReference> Widgets)
+void FWidgetBlueprintEditorUtils::BuildReplaceWithMenu(FMenuBuilder& Menu, TSharedRef<FWidgetBlueprintEditor> BlueprintEditor, UWidgetBlueprint* BP, TSet<FWidgetReference> Widgets)
 {
 	Menu.BeginSection("ReplaceWith", LOCTEXT("WidgetTree_ReplaceWith", "Replace With..."));
 	{
@@ -544,6 +547,21 @@ void FWidgetBlueprintEditorUtils::BuildReplaceWithMenu(FMenuBuilder& Menu, UWidg
 		{
 			FWidgetReference Widget = *Widgets.CreateIterator();
 			UClass* WidgetClass = Widget.GetTemplate()->GetClass();
+			TWeakObjectPtr<UClass> TemplateWidget = BlueprintEditor->GetSelectedTemplate();
+			FAssetData SelectedUserWidget = BlueprintEditor->GetSelectedUserWidget();
+			if (TemplateWidget.IsValid() || SelectedUserWidget.ObjectPath != NAME_None)
+			{
+				Menu.AddMenuEntry(
+					FText::Format(LOCTEXT("WidgetTree_ReplaceWithSelection", "Replace With {0}"), FText::FromString(TemplateWidget.IsValid() ? TemplateWidget->GetName() : SelectedUserWidget.AssetName.ToString())),
+					FText::Format(LOCTEXT("WidgetTree_ReplaceWithSelectionToolTip", "Replace this widget with a {0}"), FText::FromString(TemplateWidget.IsValid() ? TemplateWidget->GetName() : SelectedUserWidget.AssetName.ToString())),
+					FSlateIcon(),
+					FUIAction(
+						FExecuteAction::CreateStatic(&FWidgetBlueprintEditorUtils::ReplaceWidgetWithSelectedTemplate, BlueprintEditor, BP, Widget),
+						FCanExecuteAction::CreateStatic(&FWidgetBlueprintEditorUtils::CanBeReplacedWithTemplate, BlueprintEditor, BP, Widget)
+					));
+				Menu.AddMenuSeparator();
+			}
+
 			if ( WidgetClass->IsChildOf(UPanelWidget::StaticClass()) && Cast<UPanelWidget>(Widget.GetTemplate())->GetChildrenCount() == 1 )
 			{
 				Menu.AddMenuEntry(
@@ -551,7 +569,7 @@ void FWidgetBlueprintEditorUtils::BuildReplaceWithMenu(FMenuBuilder& Menu, UWidg
 					LOCTEXT("ReplaceWithChildTooltip", "Remove this widget and insert the children of this widget into the parent."),
 					FSlateIcon(),
 					FUIAction(
-						FExecuteAction::CreateStatic(&FWidgetBlueprintEditorUtils::ReplaceWidgetWithChildren, BP, Widget),
+						FExecuteAction::CreateStatic(&FWidgetBlueprintEditorUtils::ReplaceWidgetWithChildren, BlueprintEditor, BP, Widget),
 						FCanExecuteAction()
 					));
 
@@ -585,15 +603,135 @@ void FWidgetBlueprintEditorUtils::BuildReplaceWithMenu(FMenuBuilder& Menu, UWidg
 				FText::GetEmpty(),
 				FSlateIcon(),
 				FUIAction(
-					FExecuteAction::CreateStatic(&FWidgetBlueprintEditorUtils::ReplaceWidgets, BP, Widgets, ReplacementClass),
-					FCanExecuteAction::CreateStatic(&FWidgetBlueprintEditorUtils::CanBeReplaced, Widgets)
+					FExecuteAction::CreateStatic(&FWidgetBlueprintEditorUtils::ReplaceWidgets, BlueprintEditor, BP, Widgets, ReplacementClass)
 				));
 		}
 	}
 	Menu.EndSection();
 }
 
-void FWidgetBlueprintEditorUtils::ReplaceWidgetWithChildren(UWidgetBlueprint* BP, FWidgetReference Widget)
+void FWidgetBlueprintEditorUtils::ReplaceWidgetWithSelectedTemplate(TSharedRef<FWidgetBlueprintEditor> BlueprintEditor, UWidgetBlueprint* BP, FWidgetReference Widget)
+{
+	// @Todo: Needs to deal with bound object in animation tracks
+
+	const FScopedTransaction Transaction(LOCTEXT("ReplaceWidgets", "Replace Widgets"));
+	bool bIsUserWidget = false;
+
+	UWidget* ThisWidget = Widget.GetTemplate();
+	UWidget* NewReplacementWidget;
+	if (BlueprintEditor->GetSelectedTemplate().IsValid())
+	{
+		UClass* WidgetClass = BlueprintEditor->GetSelectedTemplate().Get();
+		TSharedPtr<FWidgetTemplateClass> Template = MakeShareable(new FWidgetTemplateClass(WidgetClass));
+		NewReplacementWidget = Template->Create(BP->WidgetTree);
+	}
+	else if (BlueprintEditor->GetSelectedUserWidget().ObjectPath != NAME_None)
+	{
+		bIsUserWidget = true;
+		FAssetData WidgetAssetData = BlueprintEditor->GetSelectedUserWidget();
+		TSharedPtr<FWidgetTemplateBlueprintClass> Template = MakeShareable(new FWidgetTemplateBlueprintClass(WidgetAssetData));
+		NewReplacementWidget = Template->Create(BP->WidgetTree);
+	}
+	else
+	{
+		return;
+	}
+
+	if (UPanelWidget* ExisitingPanel = Cast<UPanelWidget>(ThisWidget))
+	{
+		// if they are both panel widgets then call the existing replace function
+		UPanelWidget* ReplacementPanelWidget = Cast<UPanelWidget>(NewReplacementWidget);
+		if (ReplacementPanelWidget)
+		{
+			TSet<FWidgetReference> WidgetToReplace;
+			WidgetToReplace.Add(Widget);
+			ReplaceWidgets(BlueprintEditor, BP, WidgetToReplace, ReplacementPanelWidget->GetClass());
+			return;
+		}
+	}
+	
+	ThisWidget->Modify();
+
+	if (UPanelWidget* CurentParent = ThisWidget->GetParent())
+	{
+		CurentParent->Modify();
+		CurentParent->ReplaceChild(ThisWidget, NewReplacementWidget);
+	}
+	else if (ThisWidget == BP->WidgetTree->RootWidget)
+	{
+		BP->WidgetTree->Modify();
+		BP->WidgetTree->RootWidget = NewReplacementWidget;
+	}
+	else
+	{
+		return;
+	}
+	
+	if (!bIsUserWidget)
+	{
+		// Rename the removed widget to the transient package so that it doesn't conflict with future widgets sharing the same name.
+		ThisWidget->Rename(nullptr, nullptr);
+	}
+
+	FBlueprintEditorUtils::MarkBlueprintAsStructurallyModified(BP);
+}
+
+bool FWidgetBlueprintEditorUtils::CanBeReplacedWithTemplate(TSharedRef<FWidgetBlueprintEditor> BlueprintEditor, UWidgetBlueprint* BP, FWidgetReference Widget)
+{
+	FAssetData SelectedUserWidget = BlueprintEditor->GetSelectedUserWidget();
+	UWidget* ThisWidget = Widget.GetTemplate();
+	UPanelWidget* ExistingPanel = Cast<UPanelWidget>(ThisWidget);
+	
+	// If selecting another widget blueprint
+	if (SelectedUserWidget.ObjectPath != NAME_None)
+	{
+		if (ExistingPanel)
+		{
+			if (ExistingPanel->GetChildrenCount() != 0)
+			{
+				return false;
+			}
+		}
+		return true;
+	}
+
+	UClass* WidgetClass = BlueprintEditor->GetSelectedTemplate().Get();
+	TSharedPtr<FWidgetTemplateClass> Template = MakeShareable(new FWidgetTemplateClass(WidgetClass));
+	UPanelWidget* NewReplacementPanel = Cast<UPanelWidget>(Template->Create(BP->WidgetTree));
+
+	if (!ExistingPanel && !NewReplacementPanel)
+	{
+		return true;
+	}
+	else if (!ExistingPanel && NewReplacementPanel)
+	{
+		return true;
+	}
+	else if (ExistingPanel && !NewReplacementPanel)
+	{
+		if (ExistingPanel->GetChildrenCount() == 0)
+		{
+			return true;
+		}
+		else 
+		{
+			return false;
+		}
+	}
+	else 
+	{
+		if (ExistingPanel->GetClass()->GetDefaultObject<UPanelWidget>()->CanHaveMultipleChildren())
+		{
+			return NewReplacementPanel->GetClass()->GetDefaultObject<UPanelWidget>()->CanHaveMultipleChildren() || ExistingPanel->GetChildrenCount() == 0;
+		}
+		else
+		{
+			return true;
+		}
+	}
+}
+
+void FWidgetBlueprintEditorUtils::ReplaceWidgetWithChildren(TSharedRef<FWidgetBlueprintEditor> BlueprintEditor, UWidgetBlueprint* BP, FWidgetReference Widget)
 {
 	if ( UPanelWidget* ExistingPanelTemplate = Cast<UPanelWidget>(Widget.GetTemplate()) )
 	{
@@ -631,21 +769,7 @@ void FWidgetBlueprintEditorUtils::ReplaceWidgetWithChildren(UWidgetBlueprint* BP
 	}
 }
 
-bool FWidgetBlueprintEditorUtils::CanBeReplaced(TSet<FWidgetReference> Widgets)
-{
-	for ( FWidgetReference& Item : Widgets )
-	{
-		UPanelWidget* ExisitingPanel = Cast<UPanelWidget>( Item.GetTemplate() );
-		if ( !ExisitingPanel )
-		{
-			return false;
-		}
-	}
-
-	return true;
-}
-
-void FWidgetBlueprintEditorUtils::ReplaceWidgets(UWidgetBlueprint* BP, TSet<FWidgetReference> Widgets, UClass* WidgetClass)
+void FWidgetBlueprintEditorUtils::ReplaceWidgets(TSharedRef<FWidgetBlueprintEditor> BlueprintEditor, UWidgetBlueprint* BP, TSet<FWidgetReference> Widgets, UClass* WidgetClass)
 {
 	const FScopedTransaction Transaction(LOCTEXT("ReplaceWidgets", "Replace Widgets"));
 
@@ -653,16 +777,14 @@ void FWidgetBlueprintEditorUtils::ReplaceWidgets(UWidgetBlueprint* BP, TSet<FWid
 
 	for ( FWidgetReference& Item : Widgets )
 	{
-		if ( UPanelWidget* ExistingPanel = Cast<UPanelWidget>(Item.GetTemplate()) )
-		{
 			UPanelWidget* NewReplacementWidget = CastChecked<UPanelWidget>(Template->Create(BP->WidgetTree));
 
-			ExistingPanel->Modify();
+			Item.GetTemplate()->Modify();
 
-			if ( UPanelWidget* CurrentParent = ExistingPanel->GetParent() )
+			if ( UPanelWidget* CurrentParent = Item.GetTemplate()->GetParent() )
 			{
 				CurrentParent->Modify();
-				CurrentParent->ReplaceChild(ExistingPanel, NewReplacementWidget);
+				CurrentParent->ReplaceChild(Item.GetTemplate(), NewReplacementWidget);
 			}
 			else if ( Item.GetTemplate() == BP->WidgetTree->RootWidget )
 			{
@@ -674,17 +796,19 @@ void FWidgetBlueprintEditorUtils::ReplaceWidgets(UWidgetBlueprint* BP, TSet<FWid
 				continue;
 			}
 
-			while ( ExistingPanel->GetChildrenCount() > 0 )
+			if (UPanelWidget* ExistingPanel = Cast<UPanelWidget>(Item.GetTemplate()))
 			{
-				UWidget* Widget = ExistingPanel->GetChildAt(0);
-				Widget->Modify();
+				while (ExistingPanel->GetChildrenCount() > 0)
+				{
+					UWidget* Widget = ExistingPanel->GetChildAt(0);
+					Widget->Modify();
 
-				NewReplacementWidget->AddChild(Widget);
+					NewReplacementWidget->AddChild(Widget);
+				}
 			}
 
 			// Rename the removed widget to the transient package so that it doesn't conflict with future widgets sharing the same name.
-			ExistingPanel->Rename(nullptr, nullptr);
-		}
+			Item.GetTemplate()->Rename(nullptr, nullptr);
 	}
 
 	FBlueprintEditorUtils::MarkBlueprintAsStructurallyModified(BP);
