Translation Picker improvements

 - Fixed the wrong culture being used when editing translations.
 - Fixed being unable to edit translations that contained package localization IDs.
 - Fixed crash when the UTranslationUnit being picked was GC'd.
 - Fixed places performing case-insensitive comparisons.
 - Improved the layout and made all fields copyable.

diff --git a/Engine/Source/Editor/TranslationEditor/Private/TranslationDataManager.cpp b/Engine/Source/Editor/TranslationEditor/Private/TranslationDataManager.cpp
--- a/Engine/Source/Editor/TranslationEditor/Private/TranslationDataManager.cpp
+++ b/Engine/Source/Editor/TranslationEditor/Private/TranslationDataManager.cpp
@@ -786,7 +786,7 @@ void FTranslationDataManager::LoadFromArchive(TArray<UTranslationUnit*>& InTrans
 							FString CurrentTranslationTrimmed = TranslationUnit->Translation;
 							CurrentTranslationTrimmed.TrimStartAndEndInline();
 							// Ignore changes to only whitespace at beginning and/or end of string on import
-							if (PreviousTranslationTrimmed == CurrentTranslationTrimmed)
+							if (PreviousTranslationTrimmed.Equals(CurrentTranslationTrimmed))
 							{
 								TranslationUnit->Translation = PreviousTranslation;
 							}
@@ -929,10 +929,10 @@ bool FTranslationDataManager::SaveSelectedTranslations(TArray<UTranslationUnit*>
 					for (UTranslationUnit* Translation : TranslationsArray)
 					{
 						// If namespace matches...
-						if (Translation->Namespace == EditedItem->Namespace)
+						if (Translation->Namespace.Equals(EditedItem->Namespace))
 						{
 							// And source matches
-							if (Translation->Source == EditedItem->Source)
+							if (Translation->Source.Equals(EditedItem->Source))
 							{
 								// Update the translation in TranslationDataManager, and finish searching these translations
 								Translation->Translation = EditedItem->Translation;
diff --git a/Engine/Source/Editor/TranslationEditor/Private/TranslationPickerEditWindow.cpp b/Engine/Source/Editor/TranslationEditor/Private/TranslationPickerEditWindow.cpp
--- a/Engine/Source/Editor/TranslationEditor/Private/TranslationPickerEditWindow.cpp
+++ b/Engine/Source/Editor/TranslationEditor/Private/TranslationPickerEditWindow.cpp
@@ -4,6 +4,7 @@
 #include "Misc/Paths.h"
 #include "Misc/CommandLine.h"
 #include "Internationalization/Culture.h"
+#include "Internationalization/TextNamespaceUtil.h"
 #include "Widgets/SBoxPanel.h"
 #include "Styling/SlateTypes.h"
 #include "Framework/Application/SlateApplication.h"
@@ -31,7 +32,7 @@ const int32 STranslationPickerEditWindow::DefaultEditWindowWidth = 500;
 const int32 STranslationPickerEditWindow::DefaultEditWindowHeight = 500;
 
 UTranslationPickerSettings::UTranslationPickerSettings(const FObjectInitializer& ObjectInitializer)
-: Super(ObjectInitializer)
+	: Super(ObjectInitializer)
 {
 
 }
@@ -84,47 +85,34 @@ void STranslationPickerEditWindow::Construct(const FArguments& InArgs)
 		SNew(SBorder)
 		.BorderImage(FEditorStyle::GetBrush("ToolPanel.GroupBorder"))
 		[
-
 			SNew(SVerticalBox)
-			+ SVerticalBox::Slot()
-			.AutoHeight()
-			.HAlign(HAlign_Center)
-			.Padding(DefaultPadding)
-			[
-				// Display name of the current language
-				SNew(STextBlock)
-				.Text(FText::FromString(FInternationalization::Get().GetCurrentCulture()->GetDisplayName()))
-				.Justification(ETextJustify::Center)
-			]
-			+ SVerticalBox::Slot()
+
+			+SVerticalBox::Slot()
 			[
 				SNew(SScrollBox)
-				+ SScrollBox::Slot()
-				.Padding(0.0f)
+				+SScrollBox::Slot()
+				.Padding(FMargin(8, 5, 8, 5))
 				[
-					SNew(SVerticalBox)
-					+ SVerticalBox::Slot()
-					.AutoHeight()
-					.Padding(FMargin(8, 5, 8, 5))
-					[
-						TextsBox
-					]
+					TextsBox
 				]
 			]
-			+ SVerticalBox::Slot()
+			
+			+SVerticalBox::Slot()
 			.AutoHeight()
 			.HAlign(HAlign_Fill)
 			.Padding(DefaultPadding)
 			[
 				SNew(SVerticalBox)
-				+ SVerticalBox::Slot()
+				
+				+SVerticalBox::Slot()
 				.AutoHeight()
 				.HAlign(HAlign_Left)
 				.Padding(DefaultPadding)
 				[
 					SNew(SHorizontalBox)
 					.Visibility(bShowLocServiceCheckbox ? EVisibility::Visible : EVisibility::Collapsed)
-					+ SHorizontalBox::Slot()
+					
+					+SHorizontalBox::Slot()
 					.Padding(FMargin(3, 3, 3, 3))
 					.VAlign(VAlign_Center)
 					.AutoWidth()
@@ -141,7 +129,8 @@ void STranslationPickerEditWindow::Construct(const FArguments& InArgs)
 						}
 						)
 					]
-					+ SHorizontalBox::Slot()
+					
+					+SHorizontalBox::Slot()
 					.Padding(FMargin(0, 0, 3, 0))
 					.VAlign(VAlign_Center)
 					.AutoWidth()
@@ -151,39 +140,34 @@ void STranslationPickerEditWindow::Construct(const FArguments& InArgs)
 						.ToolTipText(LOCTEXT("SubmitTranslationPickerChangesToLocalizationServiceToolTip", "Submit changes to localization service"))
 					]
 				]
-				+ SVerticalBox::Slot()
+				
+				+SVerticalBox::Slot()
 				.AutoHeight()
 				.HAlign(HAlign_Right)
-				.Padding(DefaultPadding)
+				.Padding(FMargin(0, 5))
 				[
 					SNew(SUniformGridPanel)
 					.SlotPadding(FEditorStyle::GetMargin("StandardDialog.SlotPadding"))
 					.MinDesiredSlotWidth(FEditorStyle::GetFloat("StandardDialog.MinDesiredSlotWidth"))
 					.MinDesiredSlotHeight(FEditorStyle::GetFloat("StandardDialog.MinDesiredSlotHeight"))
-					+ SUniformGridPanel::Slot(0, 0)
+					
+					+SUniformGridPanel::Slot(0, 0)
 					[
 						SNew(SButton)
 						.HAlign(HAlign_Center)
+						.VAlign(VAlign_Center)
 						.ContentPadding(FEditorStyle::GetMargin("StandardDialog.ContentPadding"))
 						.OnClicked(this, &STranslationPickerEditWindow::SaveAllAndClose)
-						[
-							SNew(SHorizontalBox)
-							+ SHorizontalBox::Slot()
-							.Padding(FMargin(0, 0, 3, 0))
-							.VAlign(VAlign_Center)
-							.AutoWidth()
-							[
-								SNew(STextBlock).Text(LOCTEXT("SaveAllAndClose", "Save all and close"))
-							]
-						]
+						.Text(LOCTEXT("SaveAllAndClose", "Save All and Close"))
 					]
-					+ SUniformGridPanel::Slot(1, 0)
+					
+					+SUniformGridPanel::Slot(1, 0)
 					[
 						SNew(SButton)
 						.HAlign(HAlign_Center)
+						.VAlign(VAlign_Center)
 						.ContentPadding(FEditorStyle::GetMargin("StandardDialog.ContentPadding"))
 						.OnClicked(this, &STranslationPickerEditWindow::Close)
-						.VAlign(VAlign_Center)
 						.Text(LOCTEXT("CancelButton", "Cancel"))
 					]
 				]
@@ -275,31 +259,17 @@ void STranslationPickerEditWidget::Construct(const FArguments& InArgs)
 	FString ArchiveFilePath = FPaths::GetPath(LocresFullPath);
 	FString LocResCultureName = FPaths::GetBaseFilename(ArchiveFilePath);
 
-	FInternationalization& I18N = FInternationalization::Get();
-	FString CurrentCultureName = I18N.GetCurrentCulture()->GetName();
-
-	// Make sure locres path matches the current culture (can change if using the culture=__ console command)
-	if (!LocResCultureName.IsEmpty() && CurrentCultureName != LocResCultureName)
-	{
-		FString BasePath = FPaths::GetPath(ArchiveFilePath);
-		LocresFullPath = BasePath / CurrentCultureName / ManifestAndArchiveNameString + ".locres";
-	}
-
-	FText Namespace = FText::FromString(NamespaceString.Get(TEXT("")));
+	const FString CleanNamespaceString = TextNamespaceUtil::StripPackageNamespace(NamespaceString.Get(TEXT("")));
+	FText Namespace = FText::FromString(CleanNamespaceString);
 	FText Key = FText::FromString(KeyString.Get(TEXT("")));
 	FText Source = SourceString != nullptr ? FText::FromString(*SourceString) : FText::GetEmpty();
 	FText ManifestAndArchiveName = FText::FromString(ManifestAndArchiveNameString);
 	FText Translation = FText::FromString(TranslationString);
 
-	FText SourceLabel = LOCTEXT("SourceLabel", "Source:");
-	FText NamespaceLabel = LOCTEXT("NamespaceLabel", "Namespace:");
-	FText KeyLabel = LOCTEXT("KeyLabel", "Key:");
-	FText ManifestAndArchiveNameLabel = LOCTEXT("LocresFileLabel", "Target :");
-
 	// Save the necessary data in UTranslationUnit for later.  This is what we pass to TranslationDataManager to save our edits
 	TranslationUnit = NewObject<UTranslationUnit>();
-	TranslationUnit->Namespace = NamespaceString.Get(TEXT(""));
-	TranslationUnit->Source = SourceString != nullptr ? *SourceString : "";
+	TranslationUnit->Namespace = CleanNamespaceString;
+	TranslationUnit->Source = SourceString != nullptr ? *SourceString : TEXT("");
 	TranslationUnit->Translation = TranslationString;
 	TranslationUnit->LocresPath = LocresFullPath;
 
@@ -307,177 +277,172 @@ void STranslationPickerEditWidget::Construct(const FArguments& InArgs)
 	bHasRequiredLocalizationInfoForSaving = NamespaceString.IsSet() && SourceString != nullptr && LocresFullPath.Len() > 0;
 
 	TSharedPtr<SGridPanel> GridPanel;
-	TSharedRef<SGridPanel> LocalizationInfoAndSaveButtonSlot = SNew(SGridPanel).FillColumn(2,1);
 
 	// Layout all our data
 	ChildSlot
+	[
+		SNew(SHorizontalBox)
+		
+		+SHorizontalBox::Slot()
+		.FillWidth(1)
+		.Padding(FMargin(5))
 		[
-			SNew(SHorizontalBox)
-			+ SHorizontalBox::Slot()
-			.FillWidth(1)
-			.Padding(FMargin(5))
+			SNew(SVerticalBox)
+			+SVerticalBox::Slot()
 			[
-				SNew(SVerticalBox)
-				+ SVerticalBox::Slot()
-				[
 				SAssignNew(GridPanel, SGridPanel)
 				.FillColumn(1,1)
-				+ SGridPanel::Slot(0,0)
-				.Padding(FMargin(5))
+				
+				+SGridPanel::Slot(0,0)
+				.Padding(FMargin(2.5))
 				.HAlign(HAlign_Right)
 				[
-						SNew(STextBlock)
-						.TextStyle(FEditorStyle::Get(), "RichTextBlock.Bold")
-						.Text(SourceLabel)
+					SNew(STextBlock)
+					.TextStyle(FEditorStyle::Get(), "RichTextBlock.Bold")
+					.Text(LOCTEXT("SourceLabel", "Source:"))
 				]
-				+ SGridPanel::Slot(0, 1)
-					.Padding(FMargin(5))
-					.HAlign(HAlign_Right)
+
+				+SGridPanel::Slot(0, 1)
+				.Padding(FMargin(2.5))
+				.HAlign(HAlign_Right)
 				[
-					SNew(SVerticalBox)
+					SNew(SBox)
 					// Hide translation if we don't have necessary information to modify, and is same as source
 					.Visibility(!bHasRequiredLocalizationInfoForSaving && SourceString->Equals(TranslationString) ? EVisibility::Collapsed : EVisibility::Visible)
-					+ SVerticalBox::Slot()
 					[
 						SNew(STextBlock)
 						.TextStyle(FEditorStyle::Get(), "RichTextBlock.Bold")
-						.Text(LOCTEXT("TranslationLabel", "Translation: "))
+						.Text(FText::Format(LOCTEXT("TranslationLabel", "Translation ({0}):"), FText::FromString(LocResCultureName)))
 					]
 				]
 				
-				+ SGridPanel::Slot(1, 0)
-					.Padding(FMargin(5))
-					.ColumnSpan(2)
+				+SGridPanel::Slot(1, 0)
+				.Padding(FMargin(2.5))
 				[
-					SNew(STextBlock)
+					SNew(SMultiLineEditableTextBox)
+					.IsReadOnly(true)
 					.Text(Source)
 				]
-				+ SGridPanel::Slot(1, 1)
-					.ColumnSpan(2)
-					.Padding(FMargin(5))
+
+				+SGridPanel::Slot(1, 1)
+				.Padding(FMargin(2.5))
+				[
+					SNew(SBox)
+					// Hide translation if we don't have necessary information to modify, and is same as source
+					.Visibility(!bHasRequiredLocalizationInfoForSaving && SourceString->Equals(TranslationString) ? EVisibility::Collapsed : EVisibility::Visible)
 					[
-						SNew(SVerticalBox)
-						// Hide translation if we don't have necessary information to modify, and is same as source
-						.Visibility(!bHasRequiredLocalizationInfoForSaving && SourceString->Equals(TranslationString) ? EVisibility::Collapsed : EVisibility::Visible)
-						+ SVerticalBox::Slot()
-						[
-							SAssignNew(TextBox, SMultiLineEditableTextBox)
-							.IsEnabled(bAllowEditing && bHasRequiredLocalizationInfoForSaving)
-							.Text(Translation)
-							.HintText(LOCTEXT("TranslationEditTextBox_HintText", "Enter/edit translation here."))
-						]
+						SAssignNew(TextBox, SMultiLineEditableTextBox)
+						.IsReadOnly(!bAllowEditing || !bHasRequiredLocalizationInfoForSaving)
+						.Text(Translation)
+						.HintText(LOCTEXT("TranslationEditTextBox_HintText", "Enter/edit translation here."))
 					]
 				]
-			+ SVerticalBox::Slot()
-				.AutoHeight()
-				.Padding(FMargin(5))
-				[
-					LocalizationInfoAndSaveButtonSlot
-				]
 			]
-		];
+		]
+	];
 
 	if (bCultureInvariant)
 	{
-		LocalizationInfoAndSaveButtonSlot->AddSlot(0, 0)
-			.Padding(FMargin(5))
+		GridPanel->AddSlot(0, 2)
+			.Padding(FMargin(2.5))
 			.ColumnSpan(2)
+			.HAlign(HAlign_Center)
 			[
 				SNew(STextBlock)
 				.Text(LOCTEXT("CultureInvariantLabel", "This text is culture-invariant"))
-				.Justification(ETextJustify::Center)
 			];
 	}
 	else if (!bShouldGatherForLocalization)
 	{
-		LocalizationInfoAndSaveButtonSlot->AddSlot(0, 0)
-			.Padding(FMargin(5))
+		GridPanel->AddSlot(0, 2)
+			.Padding(FMargin(2.5))
 			.ColumnSpan(2)
+			.HAlign(HAlign_Center)
 			[
 				SNew(STextBlock)
 				.Text(LOCTEXT("NotGatheredForLocalizationLabel", "This text is not gathered for localization"))
-				.Justification(ETextJustify::Center)
 			];
 	}
 	else if (!bHasRequiredLocalizationInfoForSaving)
 	{
-		LocalizationInfoAndSaveButtonSlot->AddSlot(0, 0)
-			.Padding(FMargin(5))
+		GridPanel->AddSlot(0, 2)
+			.Padding(FMargin(2.5))
 			.ColumnSpan(2)
+			.HAlign(HAlign_Center)
 			[
 				SNew(STextBlock)
 				.Text(LOCTEXT("RequiredLocalizationInfoNotFound", "This text is not ready to be localized."))
-				.Justification(ETextJustify::Center)
 			];
 	}
 	else
 	{
-		LocalizationInfoAndSaveButtonSlot->AddSlot(0, 0)
+		GridPanel->AddSlot(0, 2)
 			.Padding(FMargin(2.5))
 			.HAlign(HAlign_Right)
 			[
 				SNew(STextBlock)
 				.TextStyle(FEditorStyle::Get(), "RichTextBlock.Bold")
-				.Text(NamespaceLabel)
+				.Text(LOCTEXT("NamespaceLabel", "Namespace:"))
 			];
-		LocalizationInfoAndSaveButtonSlot->AddSlot(1, 0)
+		GridPanel->AddSlot(1, 2)
 			.Padding(FMargin(2.5))
 			[
-				SNew(STextBlock)
+				SNew(SEditableTextBox)
+				.IsReadOnly(true)
 				.Text(Namespace)
 			];
-		LocalizationInfoAndSaveButtonSlot->AddSlot(0, 1)
+		GridPanel->AddSlot(0, 3)
 			.Padding(FMargin(2.5))
 			.HAlign(HAlign_Right)
 			[
 				SNew(STextBlock)
 				.TextStyle(FEditorStyle::Get(), "RichTextBlock.Bold")
-				.Text(KeyLabel)
+				.Text(LOCTEXT("KeyLabel", "Key:"))
 			];
-		LocalizationInfoAndSaveButtonSlot->AddSlot(1, 1)
+		GridPanel->AddSlot(1, 3)
 			.Padding(FMargin(2.5))
 			[
-				SNew(STextBlock)
+				SNew(SEditableTextBox)
+				.IsReadOnly(true)
 				.Text(Key)
 			];
-		LocalizationInfoAndSaveButtonSlot->AddSlot(0, 2)
+		GridPanel->AddSlot(0, 4)
 			.Padding(FMargin(2.5))
 			.HAlign(HAlign_Right)
 			[
 				SNew(STextBlock)
 				.TextStyle(FEditorStyle::Get(), "RichTextBlock.Bold")
-				.Text(ManifestAndArchiveNameLabel)
+				.Text(LOCTEXT("LocresFileLabel", "Target:"))
 			];
-		LocalizationInfoAndSaveButtonSlot->AddSlot(1, 2)
+		GridPanel->AddSlot(1, 4)
 			.Padding(FMargin(2.5))
 			[
-				SNew(STextBlock)
+				SNew(SEditableTextBox)
+				.IsReadOnly(true)
 				.Text(ManifestAndArchiveName)
 			];
-		LocalizationInfoAndSaveButtonSlot->AddSlot(2, 2)
+		GridPanel->AddSlot(0, 5)
 			.Padding(FMargin(2.5))
+			.ColumnSpan(2)
+			.HAlign(HAlign_Right)
 			[
 				SNew(SButton)
 				.HAlign(HAlign_Center)
+				.VAlign(VAlign_Center)
 				.ContentPadding(FEditorStyle::GetMargin("StandardDialog.ContentPadding"))
 				.OnClicked(this, &STranslationPickerEditWidget::SaveAndPreview)
 				.IsEnabled(bHasRequiredLocalizationInfoForSaving)
 				.Visibility(bAllowEditing ? EVisibility::Visible : EVisibility::Collapsed)
-				[
-					SNew(SHorizontalBox)
-					+ SHorizontalBox::Slot()
-					.Padding(FMargin(0, 0, 3, 0))
-					.VAlign(VAlign_Center)
-					.AutoWidth()
-					[
-						SNew(STextBlock)
-						.Text(bHasRequiredLocalizationInfoForSaving ? LOCTEXT("SaveAndPreviewButtonText", "Save and preview") : LOCTEXT("SaveAndPreviewButtonDisabledText", "Cannot Save"))
-					]
-				]
+				.Text(bHasRequiredLocalizationInfoForSaving ? LOCTEXT("SaveAndPreviewButtonText", "Save and Preview") : LOCTEXT("SaveAndPreviewButtonDisabledText", "Cannot Save"))
 			];
 	}
 }
 
+void STranslationPickerEditWidget::AddReferencedObjects(FReferenceCollector& Collector)
+{
+	Collector.AddReferencedObject(TranslationUnit);
+}
+
 FReply STranslationPickerEditWidget::SaveAndPreview()
 {
 	// Update translation string from entered text
diff --git a/Engine/Source/Editor/TranslationEditor/Private/TranslationPickerEditWindow.h b/Engine/Source/Editor/TranslationEditor/Private/TranslationPickerEditWindow.h
--- a/Engine/Source/Editor/TranslationEditor/Private/TranslationPickerEditWindow.h
+++ b/Engine/Source/Editor/TranslationEditor/Private/TranslationPickerEditWindow.h
@@ -3,8 +3,8 @@
 #pragma once
 
 #include "CoreMinimal.h"
-#include "UObject/ObjectMacros.h"
 #include "UObject/Object.h"
+#include "UObject/GCObject.h"
 #include "Widgets/DeclarativeSyntaxSupport.h"
 #include "Input/Reply.h"
 #include "Widgets/SCompoundWidget.h"
@@ -76,7 +76,7 @@ private:
 };
 
 /** Translation picker edit Widget to handle the display and editing of a single selected FText */
-class STranslationPickerEditWidget : public SCompoundWidget
+class STranslationPickerEditWidget : public SCompoundWidget, public FGCObject
 {
 	SLATE_BEGIN_ARGS(STranslationPickerEditWidget) {}
 
@@ -88,6 +88,8 @@ class STranslationPickerEditWidget : public SCompoundWidget
 
 	void Construct(const FArguments& InArgs);
 
+	virtual void AddReferencedObjects(FReferenceCollector& Collector) override;
+
 	/** Return the translation unit for this text, with any modifications */
 	UTranslationUnit* GetTranslationUnitWithAnyChanges();
 
diff --git a/Engine/Source/Editor/TranslationEditor/Private/TranslationPickerFloatingWindow.cpp b/Engine/Source/Editor/TranslationEditor/Private/TranslationPickerFloatingWindow.cpp
--- a/Engine/Source/Editor/TranslationEditor/Private/TranslationPickerFloatingWindow.cpp
+++ b/Engine/Source/Editor/TranslationEditor/Private/TranslationPickerFloatingWindow.cpp
@@ -105,16 +105,8 @@ void STranslationPickerFloatingWindow::Tick( const FGeometry& AllottedGeometry,
 
 			WindowContents->SetContentWidget(
 				SNew(SVerticalBox)
-				+ SVerticalBox::Slot()
-				.Padding(0)
-				.AutoHeight()
-				.Padding(FMargin(5))
-				[
-					SNew(STextBlock)
-					.Text(FText::FromString(FInternationalization::Get().GetCurrentCulture()->GetDisplayName()))
-					.Justification(ETextJustify::Center)
-				]
-			+ SVerticalBox::Slot()
+				
+				+SVerticalBox::Slot()
 				.Padding(0)
 				.FillHeight(1)
 				.Padding(FMargin(5))
@@ -122,13 +114,15 @@ void STranslationPickerFloatingWindow::Tick( const FGeometry& AllottedGeometry,
 					SNew(SScrollBox)
 					.Orientation(EOrientation::Orient_Vertical)
 					.ScrollBarAlwaysVisible(true)
+					
 					+SScrollBox::Slot()
-						.Padding(FMargin(0))
-						[
-									TextsBox
-						]
+					.Padding(FMargin(0))
+					[
+						TextsBox
+					]
 				]
-			+ SVerticalBox::Slot()
+
+				+SVerticalBox::Slot()
 				.Padding(0)
 				.AutoHeight()
 				.Padding(FMargin(5))
