Clamp Translation Picker window to the work area (like a tooltip)

diff --git a/Engine/Source/Editor/TranslationEditor/Private/TranslationPickerFloatingWindow.cpp b/Engine/Source/Editor/TranslationEditor/Private/TranslationPickerFloatingWindow.cpp
--- a/Engine/Source/Editor/TranslationEditor/Private/TranslationPickerFloatingWindow.cpp
+++ b/Engine/Source/Editor/TranslationEditor/Private/TranslationPickerFloatingWindow.cpp
@@ -147,6 +147,9 @@ void STranslationPickerFloatingWindow::Tick( const FGeometry& AllottedGeometry,
 		// Move to opposite side of the cursor than the tool tip, so they don't overlaps
 		DesiredPosition.X -= WindowSize.X;
 
+		// Clamp to work area
+		DesiredPosition = FSlateApplication::Get().CalculateTooltipWindowPosition(FSlateRect(DesiredPosition.X, DesiredPosition.Y, DesiredPosition.X, DesiredPosition.Y), WindowSize, false);
+
 		// also kind of a hack, but this is the only way at the moment to get a 'cursor decorator' without using the drag-drop code path
 		ParentWindow.Pin()->MoveWindowTo(DesiredPosition);
 	}
@@ -295,15 +298,7 @@ FReply STranslationPickerFloatingWindow::OnKeyDown(const FGeometry& MyGeometry,
 				FSlateApplication::Get().AddWindow(NewWindow);
 			}
 
-			FVector2D WindowSize = ParentWindow.Pin()->GetSizeInScreen();
-			FVector2D DesiredPosition = FSlateApplication::Get().GetCursorPos();
-			DesiredPosition.X -= FSlateApplication::Get().GetCursorSize().X;
-			DesiredPosition.Y += FSlateApplication::Get().GetCursorSize().Y;
-
-			// Open this new Edit window in the same position
-			DesiredPosition.X -= WindowSize.X;
-
-			NewWindow->MoveWindowTo(DesiredPosition);
+			NewWindow->MoveWindowTo(ParentWindow.Pin()->GetPositionInScreen());
 		}
 
 		TranslationPickerManager::ClosePickerWindow();
