Translation Picker can now pick String Table entries

It now attempts to show you the translation information for anything that has a valid identity (including String Table entries, which were previously being skipped by the ShouldGatherForLocalization check).

In cases where we have an identity but no localization target information for the translation data, we show a message saying that the text can be localized, but needs to be gathered (and hide the target field and save button).

Also updated the messages for the "not localizable" cases.

diff --git a/Engine/Source/Editor/TranslationEditor/Private/TranslationPickerEditWindow.cpp b/Engine/Source/Editor/TranslationEditor/Private/TranslationPickerEditWindow.cpp
--- a/Engine/Source/Editor/TranslationEditor/Private/TranslationPickerEditWindow.cpp
+++ b/Engine/Source/Editor/TranslationEditor/Private/TranslationPickerEditWindow.cpp
@@ -235,47 +235,48 @@ void STranslationPickerEditWidget::Construct(const FArguments& InArgs)
 	bAllowEditing = InArgs._bAllowEditing;
 	int32 DefaultPadding = 0.0f;
 
-	bool bCultureInvariant = PickedText.IsCultureInvariant();
-	bool bShouldGatherForLocalization = FTextInspector::ShouldGatherForLocalization(PickedText);
-
-	// Get all the data we need and format it properly
-	TOptional<FString> NamespaceString = FTextInspector::GetNamespace(PickedText);
-	TOptional<FString> KeyString = FTextInspector::GetKey(PickedText);
-	const FString* SourceString = FTextInspector::GetSourceString(PickedText);
-	const FString& TranslationString = FTextInspector::GetDisplayString(PickedText);
-	FString LocresFullPath;
-
-	FString ManifestAndArchiveNameString;
-	if (NamespaceString && KeyString)
+	// Try and get the localization information for this text
+	FString Namespace;
+	FString Key;
+	FString SourceString;
+	FString TranslationString;
+	bool bHasIdentity = false;
 	{
-		FString LocResId;
-		if (FTextLocalizationManager::Get().GetLocResID(NamespaceString.GetValue(), KeyString.GetValue(), LocResId))
-		{
-			LocresFullPath = *LocResId;
-			ManifestAndArchiveNameString = FPaths::GetBaseFilename(*LocResId);
-		}
+		const FString* SourceStringPtr = FTextInspector::GetSourceString(PickedText);
+		SourceString = SourceStringPtr ? *SourceStringPtr : FString();
+
+		const FTextDisplayStringRef DisplayString = FTextInspector::GetSharedDisplayString(PickedText);
+		TranslationString = *DisplayString;
+
+		bHasIdentity = FTextLocalizationManager::Get().FindNamespaceAndKeyFromDisplayString(DisplayString, Namespace, Key);
 	}
 
-	FString ArchiveFilePath = FPaths::GetPath(LocresFullPath);
-	FString LocResCultureName = FPaths::GetBaseFilename(ArchiveFilePath);
+	// Try and find the LocRes the active translation came from
+	// We assume the LocRes is named the same as the localization target
+	FString LocResPath;
+	FString LocTargetName;
+	FString LocResCultureName;
+	if (bHasIdentity && FTextLocalizationManager::Get().GetLocResID(Namespace, Key, LocResPath))
+	{
+		LocTargetName = FPaths::GetBaseFilename(LocResPath);
 
-	const FString CleanNamespaceString = TextNamespaceUtil::StripPackageNamespace(NamespaceString.Get(TEXT("")));
-	FText Namespace = FText::FromString(CleanNamespaceString);
-	FText Key = FText::FromString(KeyString.Get(TEXT("")));
-	FText Source = SourceString != nullptr ? FText::FromString(*SourceString) : FText::GetEmpty();
-	FText ManifestAndArchiveName = FText::FromString(ManifestAndArchiveNameString);
-	FText Translation = FText::FromString(TranslationString);
+		const FString CultureFilePath = FPaths::GetPath(LocResPath);
+		LocResCultureName = FPaths::GetBaseFilename(CultureFilePath);
+	}
+
+	// Clean the package localization ID from the namespace (to mirror what the text gatherer does when scraping for translation data)
+	Namespace = TextNamespaceUtil::StripPackageNamespace(Namespace);
 
 	// Save the necessary data in UTranslationUnit for later.  This is what we pass to TranslationDataManager to save our edits
 	TranslationUnit = NewObject<UTranslationUnit>();
-	TranslationUnit->Namespace = CleanNamespaceString;
-	TranslationUnit->Key = KeyString.Get(TEXT(""));
-	TranslationUnit->Source = SourceString != nullptr ? *SourceString : TEXT("");
+	TranslationUnit->Namespace = Namespace;
+	TranslationUnit->Key = Key;
+	TranslationUnit->Source = SourceString;
 	TranslationUnit->Translation = TranslationString;
-	TranslationUnit->LocresPath = LocresFullPath;
+	TranslationUnit->LocresPath = LocResPath;
 
-	// Can only save if we have all the required information
-	bHasRequiredLocalizationInfoForSaving = NamespaceString.IsSet() && SourceString != nullptr && LocresFullPath.Len() > 0;
+	// Can only save if we have have an identity and are in a known localization target file
+	bHasRequiredLocalizationInfoForSaving = bHasIdentity && !LocTargetName.IsEmpty();
 
 	TSharedPtr<SGridPanel> GridPanel;
 
@@ -309,11 +310,11 @@ void STranslationPickerEditWidget::Construct(const FArguments& InArgs)
 				[
 					SNew(SBox)
 					// Hide translation if we don't have necessary information to modify, and is same as source
-					.Visibility(!bHasRequiredLocalizationInfoForSaving && SourceString->Equals(TranslationString) ? EVisibility::Collapsed : EVisibility::Visible)
+					.Visibility(!bHasRequiredLocalizationInfoForSaving && SourceString.Equals(TranslationString) ? EVisibility::Collapsed : EVisibility::Visible)
 					[
 						SNew(STextBlock)
 						.TextStyle(FEditorStyle::Get(), "RichTextBlock.Bold")
-						.Text(FText::Format(LOCTEXT("TranslationLabel", "Translation ({0}):"), FText::FromString(LocResCultureName)))
+						.Text(FText::Format(LOCTEXT("TranslationLabel", "Translation ({0}):"), FText::AsCultureInvariant(LocResCultureName)))
 					]
 				]
 				
@@ -322,7 +323,7 @@ void STranslationPickerEditWidget::Construct(const FArguments& InArgs)
 				[
 					SNew(SMultiLineEditableTextBox)
 					.IsReadOnly(true)
-					.Text(Source)
+					.Text(FText::AsCultureInvariant(SourceString))
 				]
 
 				+SGridPanel::Slot(1, 1)
@@ -330,11 +331,11 @@ void STranslationPickerEditWidget::Construct(const FArguments& InArgs)
 				[
 					SNew(SBox)
 					// Hide translation if we don't have necessary information to modify, and is same as source
-					.Visibility(!bHasRequiredLocalizationInfoForSaving && SourceString->Equals(TranslationString) ? EVisibility::Collapsed : EVisibility::Visible)
+					.Visibility(!bHasRequiredLocalizationInfoForSaving && SourceString.Equals(TranslationString) ? EVisibility::Collapsed : EVisibility::Visible)
 					[
 						SAssignNew(TextBox, SMultiLineEditableTextBox)
 						.IsReadOnly(!bAllowEditing || !bHasRequiredLocalizationInfoForSaving)
-						.Text(Translation)
+						.Text(FText::AsCultureInvariant(TranslationString))
 						.HintText(LOCTEXT("TranslationEditTextBox_HintText", "Enter/edit translation here."))
 					]
 				]
@@ -342,40 +343,7 @@ void STranslationPickerEditWidget::Construct(const FArguments& InArgs)
 		]
 	];
 
-	if (bCultureInvariant)
-	{
-		GridPanel->AddSlot(0, 2)
-			.Padding(FMargin(2.5))
-			.ColumnSpan(2)
-			.HAlign(HAlign_Center)
-			[
-				SNew(STextBlock)
-				.Text(LOCTEXT("CultureInvariantLabel", "This text is culture-invariant"))
-			];
-	}
-	else if (!bShouldGatherForLocalization)
-	{
-		GridPanel->AddSlot(0, 2)
-			.Padding(FMargin(2.5))
-			.ColumnSpan(2)
-			.HAlign(HAlign_Center)
-			[
-				SNew(STextBlock)
-				.Text(LOCTEXT("NotGatheredForLocalizationLabel", "This text is not gathered for localization"))
-			];
-	}
-	else if (!bHasRequiredLocalizationInfoForSaving)
-	{
-		GridPanel->AddSlot(0, 2)
-			.Padding(FMargin(2.5))
-			.ColumnSpan(2)
-			.HAlign(HAlign_Center)
-			[
-				SNew(STextBlock)
-				.Text(LOCTEXT("RequiredLocalizationInfoNotFound", "This text is not ready to be localized."))
-			];
-	}
-	else
+	if (bHasIdentity)
 	{
 		GridPanel->AddSlot(0, 2)
 			.Padding(FMargin(2.5))
@@ -390,7 +358,7 @@ void STranslationPickerEditWidget::Construct(const FArguments& InArgs)
 			[
 				SNew(SEditableTextBox)
 				.IsReadOnly(true)
-				.Text(Namespace)
+				.Text(FText::AsCultureInvariant(Namespace))
 			];
 		GridPanel->AddSlot(0, 3)
 			.Padding(FMargin(2.5))
@@ -405,36 +373,76 @@ void STranslationPickerEditWidget::Construct(const FArguments& InArgs)
 			[
 				SNew(SEditableTextBox)
 				.IsReadOnly(true)
-				.Text(Key)
-			];
-		GridPanel->AddSlot(0, 4)
-			.Padding(FMargin(2.5))
-			.HAlign(HAlign_Right)
-			[
-				SNew(STextBlock)
-				.TextStyle(FEditorStyle::Get(), "RichTextBlock.Bold")
-				.Text(LOCTEXT("LocresFileLabel", "Target:"))
-			];
-		GridPanel->AddSlot(1, 4)
-			.Padding(FMargin(2.5))
-			[
-				SNew(SEditableTextBox)
-				.IsReadOnly(true)
-				.Text(ManifestAndArchiveName)
+				.Text(FText::AsCultureInvariant(Key))
 			];
-		GridPanel->AddSlot(0, 5)
+		
+		if (bHasRequiredLocalizationInfoForSaving)
+		{
+			GridPanel->AddSlot(0, 4)
+				.Padding(FMargin(2.5))
+				.HAlign(HAlign_Right)
+				[
+					SNew(STextBlock)
+					.TextStyle(FEditorStyle::Get(), "RichTextBlock.Bold")
+					.Text(LOCTEXT("LocresFileLabel", "Target:"))
+				];
+			GridPanel->AddSlot(1, 4)
+				.Padding(FMargin(2.5))
+				[
+					SNew(SEditableTextBox)
+					.IsReadOnly(true)
+					.Text(FText::AsCultureInvariant(LocTargetName))
+				];
+			GridPanel->AddSlot(0, 5)
+				.Padding(FMargin(2.5))
+				.ColumnSpan(2)
+				.HAlign(HAlign_Right)
+				[
+					SNew(SButton)
+					.HAlign(HAlign_Center)
+					.VAlign(VAlign_Center)
+					.ContentPadding(FEditorStyle::GetMargin("StandardDialog.ContentPadding"))
+					.OnClicked(this, &STranslationPickerEditWidget::SaveAndPreview)
+					.IsEnabled(bHasRequiredLocalizationInfoForSaving)
+					.Visibility(bAllowEditing ? EVisibility::Visible : EVisibility::Collapsed)
+					.Text(bHasRequiredLocalizationInfoForSaving ? LOCTEXT("SaveAndPreviewButtonText", "Save and Preview") : LOCTEXT("SaveAndPreviewButtonDisabledText", "Cannot Save"))
+				];
+		}
+		else
+		{
+			GridPanel->AddSlot(0, 4)
+				.Padding(FMargin(2.5))
+				.ColumnSpan(2)
+				.HAlign(HAlign_Center)
+				[
+					SNew(STextBlock)
+					.Text(LOCTEXT("TextLocalizable_RequiresGather", "This text is localizable (requires gather)."))
+				];
+		}
+	}
+	else
+	{
+		FText TextNotLocalizableReason = LOCTEXT("TextNotLocalizable_Generic", "This text is not localizable.");
+		if (PickedText.IsCultureInvariant())
+		{
+			TextNotLocalizableReason = LOCTEXT("TextNotLocalizable_CultureInvariant", "This text is not localizable (culture-invariant).");
+		}
+		else if (PickedText.IsTransient())
+		{
+			TextNotLocalizableReason = LOCTEXT("TextNotLocalizable_Transient", "This text is not localizable (transient).");
+		}
+		else if (!PickedText.ShouldGatherForLocalization())
+		{
+			TextNotLocalizableReason = LOCTEXT("TextNotLocalizable_InvalidForGather", "This text is not localizable (invalid for gather).");
+		}
+
+		GridPanel->AddSlot(0, 2)
 			.Padding(FMargin(2.5))
 			.ColumnSpan(2)
-			.HAlign(HAlign_Right)
+			.HAlign(HAlign_Center)
 			[
-				SNew(SButton)
-				.HAlign(HAlign_Center)
-				.VAlign(VAlign_Center)
-				.ContentPadding(FEditorStyle::GetMargin("StandardDialog.ContentPadding"))
-				.OnClicked(this, &STranslationPickerEditWidget::SaveAndPreview)
-				.IsEnabled(bHasRequiredLocalizationInfoForSaving)
-				.Visibility(bAllowEditing ? EVisibility::Visible : EVisibility::Collapsed)
-				.Text(bHasRequiredLocalizationInfoForSaving ? LOCTEXT("SaveAndPreviewButtonText", "Save and Preview") : LOCTEXT("SaveAndPreviewButtonDisabledText", "Cannot Save"))
+				SNew(STextBlock)
+				.Text(TextNotLocalizableReason)
 			];
 	}
 }
