Fixed Translation Picker to correctly pick the text of the leaf-widget (including children and any visible tooltips)

diff --git a/Engine/Source/Editor/TranslationEditor/Private/TranslationPickerFloatingWindow.cpp b/Engine/Source/Editor/TranslationEditor/Private/TranslationPickerFloatingWindow.cpp
--- a/Engine/Source/Editor/TranslationEditor/Private/TranslationPickerFloatingWindow.cpp
+++ b/Engine/Source/Editor/TranslationEditor/Private/TranslationPickerFloatingWindow.cpp
@@ -29,25 +29,6 @@ void STranslationPickerFloatingWindow::Construct(const FArguments& InArgs)
 	];
 }
 
-void STranslationPickerFloatingWindow::GetTextFromChildWidgets(TSharedRef<SWidget> Widget)
-{
-	FChildren* Children = Widget->GetChildren();
-
-	TArray<TSharedRef<SWidget>> ChildrenArray;
-	for (int ChildIndex = 0; ChildIndex < Children->Num(); ++ChildIndex)
-	{
-		ChildrenArray.Add(Children->GetChildAt(ChildIndex));
-	}
-
-	for (int ChildIndex = 0; ChildIndex < Children->Num(); ++ChildIndex)
-	{
-		// Pull out any FText from this child widget
-		GetTextFromWidget(Children->GetChildAt(ChildIndex));
-		// Recursively search this child's children for their FText
-		GetTextFromChildWidgets(Children->GetChildAt(ChildIndex));
-	}
-}
-
 void STranslationPickerFloatingWindow::Tick( const FGeometry& AllottedGeometry, const double InCurrentTime, const float InDeltaTime )
 {
 	FWidgetPath Path = FSlateApplication::Get().LocateWindowUnderMouse(FSlateApplication::Get().GetCursorPos(), FSlateApplication::Get().GetInteractiveTopLevelWindows(), true);
@@ -58,30 +39,33 @@ void STranslationPickerFloatingWindow::Tick( const FGeometry& AllottedGeometry,
 		if (!LastTickHoveringWidgetPath.IsValid() || LastTickHoveringWidgetPath.ToWidgetPath().ToString() != Path.ToString())
 		{
 			// Clear all previous text and widgets
-			PickedTexts.Empty();
+			PickedTexts.Reset();
 
-			// Search everything under the cursor for any FText we know how to parse
-			for (int32 PathIndex = Path.Widgets.Num() - 1; PathIndex >= 0; PathIndex--)
+			// Process the leaf-widget under the cursor
+			if (Path.Widgets.Num() > 0)
 			{
 				// General Widget case
-				TSharedRef<SWidget> PathWidget = Path.Widgets[PathIndex].Widget;
-				GetTextFromWidget(PathWidget);
+				TSharedRef<SWidget> PathWidget = Path.Widgets.Last().Widget;
+				PickTextFromWidget(PathWidget);
 
 				// Tooltip case
 				TSharedPtr<IToolTip> Tooltip = PathWidget->GetToolTip();
-				//FText TooltipDescription = FText::GetEmpty();
-
 				if (Tooltip.IsValid() && !Tooltip->IsEmpty())
 				{
-					GetTextFromWidget(Tooltip->AsWidget());
+					PickTextFromWidget(Tooltip->AsWidget());
 				}
 
-				// Currently LocateWindowUnderMouse doesn't return hit-test invisible widgets
-				// So recursively search all child widgets of the deepest widget in our path in case there is hit-test invisible text in the deeper children
-				// GetTextFromWidget prevents adding the same FText twice, so shouldn't be any harm
-				if (PathIndex == (Path.Widgets.Num() - 1))
+				// Also include tooltips from parent widgets in this path (since they may be visible)
+				for (int32 ParentPathIndex = Path.Widgets.Num() - 2; ParentPathIndex >= 0; --ParentPathIndex)
 				{
-					GetTextFromChildWidgets(PathWidget);
+					TSharedRef<SWidget> ParentPathWidget = Path.Widgets[ParentPathIndex].Widget;
+
+					// Tooltip case
+					TSharedPtr<IToolTip> ParentTooltip = ParentPathWidget->GetToolTip();
+					if (ParentTooltip.IsValid() && !ParentTooltip->IsEmpty())
+					{
+						PickTextFromWidget(ParentTooltip->AsWidget());
+					}
 				}
 			}
 
@@ -157,114 +141,103 @@ void STranslationPickerFloatingWindow::Tick( const FGeometry& AllottedGeometry,
 	LastTickHoveringWidgetPath = FWeakWidgetPath(Path);
 }
 
-FText STranslationPickerFloatingWindow::GetTextFromWidget(TSharedRef<SWidget> Widget)
+void STranslationPickerFloatingWindow::PickTextFromWidget(TSharedRef<SWidget> Widget)
 {
-	FText OriginalText = FText::GetEmpty();
-
-	STextBlock& TextBlock = (STextBlock&)Widget.Get();
-
-	// Have to parse the various widget types to find the FText
-	if (Widget->GetTypeAsString() == "STextBlock")
+	auto AppendPickedTextImpl = [this](const FText& InPickedText)
 	{
-		OriginalText = TextBlock.GetText();
-	}
-	else if (Widget->GetTypeAsString() == "SToolTip")
-	{
-		SToolTip& ToolTipWidget = (SToolTip&)Widget.Get();
-		OriginalText = GetTextFromWidget(ToolTipWidget.GetContentWidget());
-		if (OriginalText.IsEmpty())
+		const bool bAlreadyPicked = PickedTexts.ContainsByPredicate([&InPickedText](const FText& InOtherPickedText)
 		{
-			OriginalText = ToolTipWidget.GetTextTooltip();
-		}
-	}
-	else if (Widget->GetTypeAsString() == "SDocumentationToolTip")
-	{
-		SDocumentationToolTip& DocumentationToolTip = (SDocumentationToolTip&)Widget.Get();
-		OriginalText = DocumentationToolTip.GetTextTooltip();
-	}
-	else if (Widget->GetTypeAsString() == "SEditableText")
-	{
-		SEditableText& EditableText = (SEditableText&)Widget.Get();
-		// Always return the hint text because that's the only thing that will be translatable
-		OriginalText = EditableText.GetHintText();
-	}
-	else if (Widget->GetTypeAsString() == "SRichTextBlock")
-	{
-		SRichTextBlock& RichTextBlock = (SRichTextBlock&)Widget.Get();
-		OriginalText = RichTextBlock.GetText();
-	}
-	else if (Widget->GetTypeAsString() == "SMultiLineEditableText")
-	{
-		SMultiLineEditableText& MultiLineEditableText = (SMultiLineEditableText&)Widget.Get();
-		// Always return the hint text because that's the only thing that will be translatable
-		OriginalText = MultiLineEditableText.GetHintText();
-	}
-	else if (Widget->GetTypeAsString() == "SMultiLineEditableTextBox")
-	{
-		SMultiLineEditableTextBox& MultiLineEditableTextBox = (SMultiLineEditableTextBox&)Widget.Get();
-		OriginalText = MultiLineEditableTextBox.GetText();
-	}
-	else if (Widget->GetTypeAsString() == "SButton")
-	{
-		SButton& Button = (SButton&)Widget.Get();
+			return InOtherPickedText.IdenticalTo(InPickedText);
+		});
 
-		// It seems like the LocateWindowUnderMouse() function will sometimes return an SButton but not the FText inside the button?
-		// So try to find the first FText child of a button just in case
-		FChildren* Children = Button.GetChildren();
-		for (int ChildIndex = 0; ChildIndex < Children->Num(); ++ChildIndex)
+		if (!bAlreadyPicked)
 		{
-			TSharedRef<SWidget> ChildWidget = Children->GetChildAt(ChildIndex);
-			OriginalText = GetTextFromWidget(ChildWidget);
-			if (!OriginalText.IsEmpty())
-			{
-				break;
-			}
+			PickedTexts.Add(InPickedText);
 		}
-	}
-	
-	if (!OriginalText.IsEmpty())
+	};
+
+	auto AppendPickedText = [this, AppendPickedTextImpl](const FText& InPickedText)
 	{
+		if (InPickedText.IsEmpty())
+		{
+			return;
+		}
+
 		// Search the text from this widget's FText::Format history to find any source text
 		TArray<FHistoricTextFormatData> HistoricFormatData;
-		FTextInspector::GetHistoricFormatData(OriginalText, HistoricFormatData);
-
-		auto PickedTextsContains = [&](const FText& InTextToTest)
-		{
-			return PickedTexts.ContainsByPredicate([&](const FText& InPickedText)
-			{
-				return InTextToTest.EqualTo(InPickedText);
-			});
-		};
+		FTextInspector::GetHistoricFormatData(InPickedText, HistoricFormatData);
 
 		if (HistoricFormatData.Num() > 0)
 		{
 			for (const FHistoricTextFormatData& HistoricFormatDataItem : HistoricFormatData)
 			{
-				if (!PickedTextsContains(HistoricFormatDataItem.SourceFmt.GetSourceText()))
-				{
-					PickedTexts.Add(HistoricFormatDataItem.SourceFmt.GetSourceText());
-				}
+				AppendPickedTextImpl(HistoricFormatDataItem.SourceFmt.GetSourceText());
 
 				for (auto It = HistoricFormatDataItem.Arguments.CreateConstIterator(); It; ++It)
 				{
 					const FFormatArgumentValue& ArgumentValue = It.Value();
-					if (ArgumentValue.GetType() == EFormatArgumentType::Text && !PickedTextsContains(ArgumentValue.GetTextValue()))
+					if (ArgumentValue.GetType() == EFormatArgumentType::Text)
 					{
-						PickedTexts.Add(ArgumentValue.GetTextValue());
+						AppendPickedTextImpl(ArgumentValue.GetTextValue());
 					}
 				}
 			}
 		}
 		else
 		{
-			if (!PickedTextsContains(OriginalText))
-			{
-				PickedTexts.Add(OriginalText);
-			}
+			AppendPickedTextImpl(InPickedText);
 		}
+	};
+
+	// Have to parse the various widget types to find the FText
+	if (Widget->GetTypeAsString() == "STextBlock")
+	{
+		STextBlock& TextBlock = (STextBlock&)Widget.Get();
+		AppendPickedText(TextBlock.GetText());
+	}
+	else if (Widget->GetTypeAsString() == "SRichTextBlock")
+	{
+		SRichTextBlock& RichTextBlock = (SRichTextBlock&)Widget.Get();
+		AppendPickedText(RichTextBlock.GetText());
+	}
+	else if (Widget->GetTypeAsString() == "SToolTip")
+	{
+		SToolTip& ToolTipWidget = (SToolTip&)Widget.Get();
+		AppendPickedText(ToolTipWidget.GetTextTooltip());
+	}
+	else if (Widget->GetTypeAsString() == "SDocumentationToolTip")
+	{
+		SDocumentationToolTip& DocumentationToolTip = (SDocumentationToolTip&)Widget.Get();
+		AppendPickedText(DocumentationToolTip.GetTextTooltip());
+	}
+	else if (Widget->GetTypeAsString() == "SEditableText")
+	{
+		SEditableText& EditableText = (SEditableText&)Widget.Get();
+		AppendPickedText(EditableText.GetText());
+		AppendPickedText(EditableText.GetHintText());
+	}
+	else if (Widget->GetTypeAsString() == "SMultiLineEditableText")
+	{
+		SMultiLineEditableText& MultiLineEditableText = (SMultiLineEditableText&)Widget.Get();
+		AppendPickedText(MultiLineEditableText.GetText());
+		AppendPickedText(MultiLineEditableText.GetHintText());
 	}
 
-	return OriginalText;
+	// Recurse into child widgets
+	PickTextFromChildWidgets(Widget);
+}
+
+void STranslationPickerFloatingWindow::PickTextFromChildWidgets(TSharedRef<SWidget> Widget)
+{
+	FChildren* Children = Widget->GetChildren();
+
+	for (int32 ChildIndex = 0; ChildIndex < Children->Num(); ++ChildIndex)
+	{
+		TSharedRef<SWidget> ChildWidget = Children->GetChildAt(ChildIndex);
+
+		// Pull out any FText from this child widget
+		PickTextFromWidget(ChildWidget);
+	}
 }
 
 FReply STranslationPickerFloatingWindow::OnKeyDown(const FGeometry& MyGeometry, const FKeyEvent& InKeyEvent)
@@ -309,5 +282,4 @@ FReply STranslationPickerFloatingWindow::OnKeyDown(const FGeometry& MyGeometry,
 	return FReply::Unhandled();
 }
 
-
 #undef LOCTEXT_NAMESPACE
diff --git a/Engine/Source/Editor/TranslationEditor/Private/TranslationPickerFloatingWindow.h b/Engine/Source/Editor/TranslationEditor/Private/TranslationPickerFloatingWindow.h
--- a/Engine/Source/Editor/TranslationEditor/Private/TranslationPickerFloatingWindow.h
+++ b/Engine/Source/Editor/TranslationEditor/Private/TranslationPickerFloatingWindow.h
@@ -28,10 +28,10 @@ private:
 	virtual void Tick( const FGeometry& AllottedGeometry, const double InCurrentTime, const float InDeltaTime ) override;
 
 	/** Pull the FText reference out of an SWidget */
-	FText GetTextFromWidget(TSharedRef<SWidget> Widget);
+	void PickTextFromWidget(TSharedRef<SWidget> Widget);
 
 	/** Pull the FText reference out of the child widgets of an SWidget */
-	void GetTextFromChildWidgets(TSharedRef<SWidget> Widget);
+	void PickTextFromChildWidgets(TSharedRef<SWidget> Widget);
 
 	/** Handle key presses */
 	virtual FReply OnKeyDown(const FGeometry& MyGeometry, const FKeyEvent& InKeyEvent) override;
