Fixed Translation Picker constantly stealing focus

diff --git a/Engine/Source/Editor/TranslationEditor/Private/TranslationPickerFloatingWindow.cpp b/Engine/Source/Editor/TranslationEditor/Private/TranslationPickerFloatingWindow.cpp
--- a/Engine/Source/Editor/TranslationEditor/Private/TranslationPickerFloatingWindow.cpp
+++ b/Engine/Source/Editor/TranslationEditor/Private/TranslationPickerFloatingWindow.cpp
@@ -3,6 +3,7 @@
 #include "TranslationPickerFloatingWindow.h"
 #include "Internationalization/Culture.h"
 #include "Framework/Application/SlateApplication.h"
+#include "Framework/Application/IInputProcessor.h"
 #include "Widgets/Layout/SBorder.h"
 #include "Widgets/Text/STextBlock.h"
 #include "Widgets/Text/SRichTextBlock.h"
@@ -18,6 +19,40 @@
 
 #define LOCTEXT_NAMESPACE "TranslationPicker"
 
+class FTranslationPickerInputProcessor : public IInputProcessor
+{
+public:
+	FTranslationPickerInputProcessor(STranslationPickerFloatingWindow* InOwner)
+		: Owner(InOwner)
+	{
+	}
+
+	void SetOwner(STranslationPickerFloatingWindow* InOwner)
+	{
+		Owner = InOwner;
+	}
+
+	virtual ~FTranslationPickerInputProcessor() = default;
+
+	virtual void Tick(const float DeltaTime, FSlateApplication& SlateApp, TSharedRef<ICursor> Cursor) override
+	{
+	}
+
+	virtual bool HandleKeyDownEvent(FSlateApplication& SlateApp, const FKeyEvent& InKeyEvent) override
+	{
+		if (Owner && InKeyEvent.GetKey() == EKeys::Escape)
+		{
+			Owner->OnEscapePressed();
+			return true;
+		}
+
+		return false;
+	}
+
+private:
+	STranslationPickerFloatingWindow* Owner;
+};
+
 void STranslationPickerFloatingWindow::Construct(const FArguments& InArgs)
 {
 	ParentWindow = InArgs._ParentWindow;
@@ -27,6 +62,22 @@ void STranslationPickerFloatingWindow::Construct(const FArguments& InArgs)
 	[
 		WindowContents.ToSharedRef()
 	];
+
+	InputProcessor = MakeShared<FTranslationPickerInputProcessor>(this);
+	FSlateApplication::Get().RegisterInputPreProcessor(InputProcessor, 0);
+}
+
+STranslationPickerFloatingWindow::~STranslationPickerFloatingWindow()
+{
+	if (InputProcessor.IsValid())
+	{
+		InputProcessor->SetOwner(nullptr);
+		if (FSlateApplication::IsInitialized())
+		{
+			FSlateApplication::Get().UnregisterInputPreProcessor(InputProcessor);
+		}
+		InputProcessor.Reset();
+	}
 }
 
 void STranslationPickerFloatingWindow::Tick( const FGeometry& AllottedGeometry, const double InCurrentTime, const float InDeltaTime )
@@ -112,15 +163,13 @@ void STranslationPickerFloatingWindow::Tick( const FGeometry& AllottedGeometry,
 				.Padding(FMargin(5))
 				[
 					SNew(STextBlock)
-					.Text(PickedTexts.Num() >= 1 ? LOCTEXT("TranslationPickerEscToEdit", "Press Esc to edit translation(s)") : LOCTEXT("TranslationPickerHoverToViewEditEscToQuit", "Hover over text to view/edit translation, or press Esc to quit"))
+					.Text(PickedTexts.Num() > 0 ? LOCTEXT("TranslationPickerEscToEdit", "Press Esc to edit translations") : LOCTEXT("TranslationPickerHoverToViewEditEscToQuit", "Hover over text to view/edit translations, or press Esc to quit"))
 					.Justification(ETextJustify::Center)
 				]
 			);
 		}
 	}
 
-	// kind of a hack, but we need to maintain keyboard focus otherwise we wont get our keypress to 'pick'
-	FSlateApplication::Get().SetKeyboardFocus(SharedThis(this), EFocusCause::SetDirectly);
 	if (ParentWindow.IsValid())
 	{
 		FVector2D WindowSize = ParentWindow.Pin()->GetSizeInScreen();
@@ -240,46 +289,39 @@ void STranslationPickerFloatingWindow::PickTextFromChildWidgets(TSharedRef<SWidg
 	}
 }
 
-FReply STranslationPickerFloatingWindow::OnKeyDown(const FGeometry& MyGeometry, const FKeyEvent& InKeyEvent)
+void STranslationPickerFloatingWindow::OnEscapePressed()
 {
-	if (InKeyEvent.GetKey() == EKeys::Escape)
+	if (PickedTexts.Num() > 0)
 	{
-		if (PickedTexts.Num() > 0)
 		// Open a different window to allow editing of the translation
-		{
-			TSharedRef<SWindow> NewWindow = SNew(SWindow)
-				.Title(LOCTEXT("TranslationPickerEditWindowTitle", "Edit Translation(s)"))
-				.CreateTitleBar(true)
-				.SizingRule(ESizingRule::UserSized);
+		TSharedRef<SWindow> NewWindow = SNew(SWindow)
+			.Title(LOCTEXT("TranslationPickerEditWindowTitle", "Edit Translations"))
+			.CreateTitleBar(true)
+			.SizingRule(ESizingRule::UserSized);
 
-			TSharedRef<STranslationPickerEditWindow> EditWindow = SNew(STranslationPickerEditWindow)
-				.ParentWindow(NewWindow)
-				.PickedTexts(PickedTexts);
+		TSharedRef<STranslationPickerEditWindow> EditWindow = SNew(STranslationPickerEditWindow)
+			.ParentWindow(NewWindow)
+			.PickedTexts(PickedTexts);
 
-			NewWindow->SetContent(EditWindow);
+		NewWindow->SetContent(EditWindow);
 
-			// Make this roughly the same size as the Edit Window, so when you press Esc to edit, the window is in basically the same size
-			NewWindow->Resize(FVector2D(STranslationPickerEditWindow::DefaultEditWindowWidth, STranslationPickerEditWindow::DefaultEditWindowHeight));
+		// Make this roughly the same size as the Edit Window, so when you press Esc to edit, the window is in basically the same size
+		NewWindow->Resize(FVector2D(STranslationPickerEditWindow::DefaultEditWindowWidth, STranslationPickerEditWindow::DefaultEditWindowHeight));
 
-			TSharedPtr<SWindow> RootWindow = FGlobalTabmanager::Get()->GetRootWindow();
-			if (RootWindow.IsValid())
-			{
-				FSlateApplication::Get().AddWindowAsNativeChild(NewWindow, RootWindow.ToSharedRef());
-			}
-			else
-			{
-				FSlateApplication::Get().AddWindow(NewWindow);
-			}
-
-			NewWindow->MoveWindowTo(ParentWindow.Pin()->GetPositionInScreen());
+		TSharedPtr<SWindow> RootWindow = FGlobalTabmanager::Get()->GetRootWindow();
+		if (RootWindow.IsValid())
+		{
+			FSlateApplication::Get().AddWindowAsNativeChild(NewWindow, RootWindow.ToSharedRef());
+		}
+		else
+		{
+			FSlateApplication::Get().AddWindow(NewWindow);
 		}
 
-		TranslationPickerManager::ClosePickerWindow();
-		
-		return FReply::Handled();
+		NewWindow->MoveWindowTo(ParentWindow.Pin()->GetPositionInScreen());
 	}
 
-	return FReply::Unhandled();
+	TranslationPickerManager::ClosePickerWindow();
 }
 
 #undef LOCTEXT_NAMESPACE
diff --git a/Engine/Source/Editor/TranslationEditor/Private/TranslationPickerFloatingWindow.h b/Engine/Source/Editor/TranslationEditor/Private/TranslationPickerFloatingWindow.h
--- a/Engine/Source/Editor/TranslationEditor/Private/TranslationPickerFloatingWindow.h
+++ b/Engine/Source/Editor/TranslationEditor/Private/TranslationPickerFloatingWindow.h
@@ -12,6 +12,7 @@
 #define LOCTEXT_NAMESPACE "TranslationPicker"
 
 class SToolTip;
+class FTranslationPickerInputProcessor;
 
 /** Translation picker floating window to show details of FText(s) under cursor, and allow in-place translation via TranslationPickerEditWindow */
 class STranslationPickerFloatingWindow : public SCompoundWidget
@@ -22,9 +23,13 @@ class STranslationPickerFloatingWindow : public SCompoundWidget
 
 	SLATE_END_ARGS()
 
+	virtual ~STranslationPickerFloatingWindow();
+
 	void Construct(const FArguments& InArgs);
 
 private:
+	friend class FTranslationPickerInputProcessor;
+
 	virtual void Tick( const FGeometry& AllottedGeometry, const double InCurrentTime, const float InDeltaTime ) override;
 
 	/** Pull the FText reference out of an SWidget */
@@ -33,14 +38,11 @@ private:
 	/** Pull the FText reference out of the child widgets of an SWidget */
 	void PickTextFromChildWidgets(TSharedRef<SWidget> Widget);
 
-	/** Handle key presses */
-	virtual FReply OnKeyDown(const FGeometry& MyGeometry, const FKeyEvent& InKeyEvent) override;
+	/** Handle escape being pressed */
+	void OnEscapePressed();
 
-	/** We need to support keyboard focus to process the 'Esc' key */
-	virtual bool SupportsKeyboardFocus() const override
-	{
-		return true;
-	}
+	/** Input processor used to capture the 'Esc' key */
+	TSharedPtr<FTranslationPickerInputProcessor> InputProcessor;
 
 	/** Handle to the window that contains this widget */
 	TWeakPtr<SWindow> ParentWindow;
