Translation Editor now shows stale translations as "Needs Review" rather than "Untranslated"

diff --git a/Engine/Source/Editor/TranslationEditor/Private/TranslationDataManager.cpp b/Engine/Source/Editor/TranslationEditor/Private/TranslationDataManager.cpp
--- a/Engine/Source/Editor/TranslationEditor/Private/TranslationDataManager.cpp
+++ b/Engine/Source/Editor/TranslationEditor/Private/TranslationDataManager.cpp
@@ -255,50 +255,25 @@ bool FTranslationDataManager::WriteTranslationData(bool bForceWrite /*= false*/)
 
 		bool bNeedsWrite = false;
 
-		for (UTranslationUnit* TranslationUnit : Untranslated)
+		auto WriteTranslationUnits = [&Archive , &bNeedsWrite](const TArray<UTranslationUnit*> InTranslationUnits)
 		{
-			if (TranslationUnit != nullptr)
+			for (UTranslationUnit* TranslationUnit : InTranslationUnits)
 			{
-				const FLocItem SearchSource(TranslationUnit->Source);
-				FString OldTranslation = Archive->FindEntryByKey(TranslationUnit->Namespace, TranslationUnit->Key, TranslationUnit->KeyMetaDataObject)->Translation.Text;
-				FString TranslationToWrite = TranslationUnit->Translation;
-				if (!TranslationToWrite.Equals(OldTranslation))
+				if (TranslationUnit)
 				{
-					Archive->SetTranslation(TranslationUnit->Namespace, TranslationUnit->Key, FLocItem(TranslationUnit->Source), FLocItem(TranslationToWrite), TranslationUnit->KeyMetaDataObject);
-					bNeedsWrite = true;
-				}
-			}
-		}
-
-		for (UTranslationUnit* TranslationUnit : Review)
-		{
-			if (TranslationUnit != nullptr)
-			{
-				const FLocItem SearchSource(TranslationUnit->Source);
-				FString OldTranslation = Archive->FindEntryByKey(TranslationUnit->Namespace, TranslationUnit->Key, TranslationUnit->KeyMetaDataObject)->Translation.Text;
-				FString TranslationToWrite = TranslationUnit->Translation;
-				if (TranslationUnit->HasBeenReviewed && !TranslationToWrite.Equals(OldTranslation))
-				{
-					Archive->SetTranslation(TranslationUnit->Namespace, TranslationUnit->Key, FLocItem(TranslationUnit->Source), FLocItem(TranslationToWrite), TranslationUnit->KeyMetaDataObject);
-					bNeedsWrite = true;
+					TSharedPtr<FArchiveEntry> OldArchiveEntry = Archive->FindEntryByKey(TranslationUnit->Namespace, TranslationUnit->Key, TranslationUnit->KeyMetaDataObject);
+					if (TranslationUnit->HasBeenReviewed && (!OldArchiveEntry || !OldArchiveEntry->Source.Text.Equals(TranslationUnit->Source, ESearchCase::CaseSensitive) || !OldArchiveEntry->Translation.Text.Equals(TranslationUnit->Translation, ESearchCase::CaseSensitive)))
+					{
+						Archive->SetTranslation(TranslationUnit->Namespace, TranslationUnit->Key, FLocItem(TranslationUnit->Source), FLocItem(TranslationUnit->Translation), TranslationUnit->KeyMetaDataObject);
+						bNeedsWrite = true;
+					}
 				}
 			}
-		}
+		};
 
-		for (UTranslationUnit* TranslationUnit : Complete)
-		{
-			if (TranslationUnit != nullptr)
-			{
-				const FLocItem SearchSource(TranslationUnit->Source);
-				FString OldTranslation = Archive->FindEntryByKey(TranslationUnit->Namespace, TranslationUnit->Key, TranslationUnit->KeyMetaDataObject)->Translation.Text;
-				FString TranslationToWrite = TranslationUnit->Translation;
-				if (!TranslationToWrite.Equals(OldTranslation))
-				{
-					Archive->SetTranslation(TranslationUnit->Namespace, TranslationUnit->Key, FLocItem(TranslationUnit->Source), FLocItem(TranslationToWrite), TranslationUnit->KeyMetaDataObject);
-					bNeedsWrite = true;
-				}
-			}
-		}
+		WriteTranslationUnits(Untranslated);
+		WriteTranslationUnits(Review);
+		WriteTranslationUnits(Complete);
 
 		bSuccess = true;
 
@@ -726,15 +701,9 @@ void FTranslationDataManager::LoadFromArchive(TArray<UTranslationUnit*>& InTrans
 					const FString PreviousTranslation = TranslationUnit->Translation;
 					TranslationUnit->Translation.Reset();
 
-					FString TranslatedString = ArchiveEntry->Translation.Text;
-					if (!ArchiveEntry->Source.Text.Equals(TranslationUnit->Source, ESearchCase::CaseSensitive))
-					{
-						// Stale translation
-						TranslatedString.Reset();
-					}
-
-					if (TranslatedString.IsEmpty())
+					if (ArchiveEntry->Translation.Text.IsEmpty())
 					{
+						// No current translation - try and find an historic one
 						bool bHasTranslationHistory = false;
 						int32 MostRecentNonNullTranslationIndex = -1;
 						int32 ContextForRecentTranslation = -1;
@@ -770,9 +739,15 @@ void FTranslationDataManager::LoadFromArchive(TArray<UTranslationUnit*>& InTrans
 							Untranslated.Add(TranslationUnit);
 						}
 					}
+					else if (!ArchiveEntry->Source.Text.Equals(TranslationUnit->Source, ESearchCase::CaseSensitive))
+					{
+						// If we have a stale translation, this goes in the Needs Review tab
+						TranslationUnit->Translation = ArchiveEntry->Translation.Text;
+						Review.Add(TranslationUnit);
+					}
 					else
 					{
-						TranslationUnit->Translation = TranslatedString;
+						TranslationUnit->Translation = ArchiveEntry->Translation.Text;
 						TranslationUnit->HasBeenReviewed = true;
 						Complete.Add(TranslationUnit);
 					}
diff --git a/Engine/Source/Editor/TranslationEditor/Private/TranslationUnit.cpp b/Engine/Source/Editor/TranslationEditor/Private/TranslationUnit.cpp
--- a/Engine/Source/Editor/TranslationEditor/Private/TranslationUnit.cpp
+++ b/Engine/Source/Editor/TranslationEditor/Private/TranslationUnit.cpp
@@ -14,5 +14,11 @@ void UTranslationUnit::PostEditChangeProperty(struct FPropertyChangedEvent& Prop
 
 	const FName Name = (PropertyChangedEvent.Property != nullptr) ? PropertyChangedEvent.Property->GetFName() : NAME_None;
 
+	if (Name == GET_MEMBER_NAME_CHECKED(UTranslationUnit, Translation))
+	{
+		// Consider modifying the translation to be an implicit review
+		HasBeenReviewed = true;
+	}
+
 	TranslationUnitPropertyChangedEvent.Broadcast(Name);
 }
