Fix for UE-90683: You can no longer delete conflicting variables

diff --git a/Engine/Source/Editor/Kismet/Private/BlueprintDetailsCustomization.cpp b/Engine/Source/Editor/Kismet/Private/BlueprintDetailsCustomization.cpp
--- a/Engine/Source/Editor/Kismet/Private/BlueprintDetailsCustomization.cpp
+++ b/Engine/Source/Editor/Kismet/Private/BlueprintDetailsCustomization.cpp
@@ -180,7 +180,7 @@ FReply FBlueprintDetails::HandleAddOrViewEventForVariable(const FName EventName,
 	UBlueprint* BlueprintObj = GetBlueprintObj();
 
 	// Find the corresponding variable property in the Blueprint
-	FObjectProperty* VariableProperty = FindField<FObjectProperty>(BlueprintObj->SkeletonGeneratedClass, PropertyName);
+	FObjectProperty* VariableProperty = FindFProperty<FObjectProperty>(BlueprintObj->SkeletonGeneratedClass, PropertyName);
 
 	if ( VariableProperty )
 	{
@@ -773,7 +773,7 @@ void FBlueprintVarActionDetails::CustomizeDetails( IDetailLayoutBuilder& DetailL
 		
 			if(!IsALocalVariable(VariableProperty))
 			{
-				OriginalProperty = FindField<FProperty>(BlueprintObj->GeneratedClass, VariableProperty->GetFName());
+				OriginalProperty = FindFProperty<FProperty>(BlueprintObj->GeneratedClass, VariableProperty->GetFName());
 			}
 			else
 			{
@@ -2166,7 +2166,7 @@ void FBlueprintVarActionDetails::OnBitmaskChanged(ECheckBoxState InNewState)
 		if (LocalBlueprint->GeneratedClass)
 		{
 			UObject* CDO = LocalBlueprint->GeneratedClass->GetDefaultObject(false);
-			FProperty* VarProperty = FindField<FProperty>(LocalBlueprint->GeneratedClass, VarName);
+			FProperty* VarProperty = FindFProperty<FProperty>(LocalBlueprint->GeneratedClass, VarName);
 
 			if (CDO != nullptr && VarProperty != nullptr)
 			{
@@ -2233,7 +2233,7 @@ void FBlueprintVarActionDetails::OnBitmaskEnumTypeChanged(TSharedPtr<FString> It
 		if (LocalBlueprint->GeneratedClass)
 		{
 			UObject* CDO = LocalBlueprint->GeneratedClass->GetDefaultObject(false);
-			FProperty* VarProperty = FindField<FProperty>(LocalBlueprint->GeneratedClass, VarName);
+			FProperty* VarProperty = FindFProperty<FProperty>(LocalBlueprint->GeneratedClass, VarName);
 
 			if (CDO != nullptr && VarProperty != nullptr)
 			{
@@ -3883,7 +3883,7 @@ FText FBlueprintGraphActionDetails::GetCurrentReplicatedEventString() const
 		uint32 NetFlags = CustomEvent->FunctionFlags & ReplicatedNetMask;
 		if (CustomEvent->IsOverride())
 		{
-			UFunction* SuperFunction = FindField<UFunction>(CustomEvent->GetBlueprint()->ParentClass, CustomEvent->CustomFunctionName);
+			UFunction* SuperFunction = FindUField<UFunction>(CustomEvent->GetBlueprint()->ParentClass, CustomEvent->CustomFunctionName);
 			check(SuperFunction != NULL);
 
 			NetFlags = SuperFunction->FunctionFlags & ReplicatedNetMask;
@@ -4349,7 +4349,7 @@ bool FBaseBlueprintGraphActionDetails::OnVerifyPinRename(UK2Node_EditablePinBase
 	{
 		// Check if the name conflicts with any of the other internal UFunction's property names (local variables and parameters).
 		const UFunction* FoundFunction = FFunctionFromNodeHelper::FunctionFromNode(InTargetNode);
-		const FProperty* ExistingProperty = FindField<const FProperty>(FoundFunction, *InNewName);
+		const FProperty* ExistingProperty = FindFProperty<const FProperty>(FoundFunction, *InNewName);
 		if (ExistingProperty)
 		{
 			OutErrorMessage = LOCTEXT("ConflictsWithProperty", "Conflicts with another local variable or function parameter!");
@@ -5881,7 +5881,7 @@ void FBlueprintComponentDetails::CustomizeDetails(IDetailLayoutBuilder& DetailLa
 	if ( FBlueprintEditorUtils::DoesSupportEventGraphs(BlueprintObj) && Nodes.Num() == 1 )
 	{
 		FName PropertyName = CachedNodePtr->GetVariableName();
-		FObjectProperty* VariableProperty = FindField<FObjectProperty>(BlueprintObj->SkeletonGeneratedClass, PropertyName);
+		FObjectProperty* VariableProperty = FindFProperty<FObjectProperty>(BlueprintObj->SkeletonGeneratedClass, PropertyName);
 
 		AddEventsCategory(DetailLayout, VariableProperty);
 	}
diff --git a/Engine/Source/Editor/Kismet/Private/BlueprintEditor.cpp b/Engine/Source/Editor/Kismet/Private/BlueprintEditor.cpp
--- a/Engine/Source/Editor/Kismet/Private/BlueprintEditor.cpp
+++ b/Engine/Source/Editor/Kismet/Private/BlueprintEditor.cpp
@@ -576,7 +576,7 @@ FSlateBrush const* FBlueprintEditor::GetVarIconAndColor(const UStruct* VarScope,
 {
 	if (VarScope != NULL)
 	{
-		FProperty* Property = FindField<FProperty>(VarScope, VarName);
+		FProperty* Property = FindFProperty<FProperty>(VarScope, VarName);
 		if (Property != NULL)
 		{
 			const UEdGraphSchema_K2* K2Schema = GetDefault<UEdGraphSchema_K2>();
@@ -5557,7 +5557,7 @@ void FBlueprintEditor::ConvertFunctionToEvent(UK2Node_FunctionEntry* SelectedCal
 		FVector2D SpawnPos = EventGraph->GetGoodPlaceForNewNode();
 		
 		// Was this function implemented in as an override?
-		UFunction* ParentFunction = FindField<UFunction>(NodeBP->ParentClass, EventName);
+		UFunction* ParentFunction = FindUField<UFunction>(NodeBP->ParentClass, EventName);
 		bool bIsOverrideFunc = Func->GetSuperFunction() || (ParentFunction != nullptr);
 		const UEdGraphSchema_K2* K2Schema = GetDefault<UEdGraphSchema_K2>();
 
@@ -5738,7 +5738,7 @@ void FBlueprintEditor::ConvertEventToFunction(UK2Node_Event* SelectedEventNode)
 			SourceGraph->Modify(); 
 
 			// Check if this is an override function
-			UFunction* ParentFunction = FindField<UFunction>(NodeBP->ParentClass, OriginalEventName);
+			UFunction* ParentFunction = FindUField<UFunction>(NodeBP->ParentClass, OriginalEventName);
 			const bool bIsOverrideFunc = FunctionSig->GetSuperFunction() || (ParentFunction != nullptr);
 			UClass* const OverrideFuncClass = FBlueprintEditorUtils::GetOverrideFunctionClass(NodeBP, OriginalEventName);
 
@@ -5868,7 +5868,7 @@ void FBlueprintEditor::ConvertEventToFunction(UK2Node_Event* SelectedEventNode)
 			FBlueprintEditorUtils::RenameGraph(NewGraph, *OriginalEventNameString);
 
 			// If this function is blueprint callable, then spawn a function call node to it in place of the old event
-			UFunction* const NewFunction = FindField<UFunction>(NodeBP->SkeletonGeneratedClass, OriginalEventName);
+			UFunction* const NewFunction = FindUField<UFunction>(NodeBP->SkeletonGeneratedClass, OriginalEventName);
 			if (NewFunction && NewFunction->HasAllFunctionFlags(FUNC_BlueprintCallable))
 			{
 				IBlueprintNodeBinder::FBindingSet Bindings;
@@ -7755,7 +7755,7 @@ UEdGraph* FBlueprintEditor::CollapseSelectionToFunction(TSharedPtr<SGraphEditor>
 	TempListBuilder.OwnerOfTemporaries->SetFlags(RF_Transient);
 
 	IBlueprintNodeBinder::FBindingSet Bindings;
-	OutFunctionNode = UBlueprintFunctionNodeSpawner::Create(FindField<UFunction>(GetBlueprintObj()->SkeletonGeneratedClass, DocumentName))->Invoke(SourceGraph, Bindings, FVector2D::ZeroVector);
+	OutFunctionNode = UBlueprintFunctionNodeSpawner::Create(FindUField<UFunction>(GetBlueprintObj()->SkeletonGeneratedClass, DocumentName))->Invoke(SourceGraph, Bindings, FVector2D::ZeroVector);
 
 	check(OutFunctionNode);
 
@@ -8192,7 +8192,7 @@ void FBlueprintEditor::OnAddNewLocalVariable()
 	UEdGraph* TargetGraph = FBlueprintEditorUtils::GetTopLevelGraph(FocusedGraphEdPtr.Pin()->GetCurrentGraph());
 	check(TargetGraph->GetSchema()->GetGraphType(TargetGraph) == GT_Function);
 
-	FName VarName = FBlueprintEditorUtils::FindUniqueKismetName(GetBlueprintObj(), TEXT("NewLocalVar"), FindField<UFunction>(GetBlueprintObj()->SkeletonGeneratedClass, TargetGraph->GetFName()));
+	FName VarName = FBlueprintEditorUtils::FindUniqueKismetName(GetBlueprintObj(), TEXT("NewLocalVar"), FindUField<UFunction>(GetBlueprintObj()->SkeletonGeneratedClass, TargetGraph->GetFName()));
 
 	bool bSuccess = MyBlueprintWidget.IsValid() && FBlueprintEditorUtils::AddLocalVariable(GetBlueprintObj(), TargetGraph, VarName, MyBlueprintWidget->GetLastPinTypeUsed());
 
diff --git a/Engine/Source/Editor/Kismet/Private/BlueprintPaletteFavorites.cpp b/Engine/Source/Editor/Kismet/Private/BlueprintPaletteFavorites.cpp
--- a/Engine/Source/Editor/Kismet/Private/BlueprintPaletteFavorites.cpp
+++ b/Engine/Source/Editor/Kismet/Private/BlueprintPaletteFavorites.cpp
@@ -401,7 +401,7 @@ void UBlueprintPaletteFavorites::RemoveFavorites(TArray< TSharedPtr<FEdGraphSche
 //------------------------------------------------------------------------------
 void UBlueprintPaletteFavorites::LoadProfile(FString const& ProfileName)
 {
-	PreEditChange(FindField<FProperty>(GetClass(), TEXT("CurrentProfile")));
+	PreEditChange(FindFProperty<FProperty>(GetClass(), TEXT("CurrentProfile")));
 	{
 		CurrentProfile = ProfileName;
 		LoadSetProfile();
@@ -490,7 +490,7 @@ void UBlueprintPaletteFavorites::LoadCustomFavorites()
 //------------------------------------------------------------------------------
 void UBlueprintPaletteFavorites::SetProfile(FString const& ProfileName)
 {
-	PreEditChange(FindField<FProperty>(GetClass(), TEXT("CurrentProfile")));
+	PreEditChange(FindFProperty<FProperty>(GetClass(), TEXT("CurrentProfile")));
 	{
 		CurrentProfile = ProfileName;
 	}
diff --git a/Engine/Source/Editor/Kismet/Private/FindInBlueprintManager.cpp b/Engine/Source/Editor/Kismet/Private/FindInBlueprintManager.cpp
--- a/Engine/Source/Editor/Kismet/Private/FindInBlueprintManager.cpp
+++ b/Engine/Source/Editor/Kismet/Private/FindInBlueprintManager.cpp
@@ -758,7 +758,7 @@ namespace BlueprintSearchMetaDataHelpers
 		SavePinTypeToJson(InWriter, VariableType);
 
 		// Find the FProperty and convert it into a Json value.
-		FProperty* VariableProperty = FindField<FProperty>(InBlueprint->GeneratedClass, InVariableDescription.VarName);
+		FProperty* VariableProperty = FindFProperty<FProperty>(InBlueprint->GeneratedClass, InVariableDescription.VarName);
 		if(VariableProperty)
 		{
 			const uint8* PropData = VariableProperty->ContainerPtrToValuePtr<uint8>(InBlueprint->GeneratedClass->GetDefaultObject());
diff --git a/Engine/Source/Editor/Kismet/Private/SBlueprintActionMenu.cpp b/Engine/Source/Editor/Kismet/Private/SBlueprintActionMenu.cpp
--- a/Engine/Source/Editor/Kismet/Private/SBlueprintActionMenu.cpp
+++ b/Engine/Source/Editor/Kismet/Private/SBlueprintActionMenu.cpp
@@ -496,7 +496,7 @@ void SBlueprintActionMenu::ConstructActionContext(FBlueprintActionContext& Conte
 			if (Nodes.Num() == 1 && Nodes[0]->GetNodeType() == FSCSEditorTreeNode::ComponentNode)
 			{
 				FName PropertyName = Nodes[0]->GetVariableName();
-				FObjectProperty* VariableProperty = FindField<FObjectProperty>(Blueprint->SkeletonGeneratedClass, PropertyName);
+				FObjectProperty* VariableProperty = FindFProperty<FObjectProperty>(Blueprint->SkeletonGeneratedClass, PropertyName);
 				ContextDescOut.SelectedObjects.Add(VariableProperty);
 			}
 		}
diff --git a/Engine/Source/Editor/Kismet/Private/SBlueprintPalette.cpp b/Engine/Source/Editor/Kismet/Private/SBlueprintPalette.cpp
--- a/Engine/Source/Editor/Kismet/Private/SBlueprintPalette.cpp
+++ b/Engine/Source/Editor/Kismet/Private/SBlueprintPalette.cpp
@@ -86,7 +86,7 @@ static FString GetVarType(UStruct* VarScope, FName VarName, bool bUseObjToolTip,
 
 	if (VarScope)
 	{
-		if (FProperty* Property = FindField<FProperty>(VarScope, VarName))
+		if (FProperty* Property = FindFProperty<FProperty>(VarScope, VarName))
 		{
 			// If it is an object property, see if we can get a nice class description instead of just the name
 			FObjectProperty* ObjProp = CastField<FObjectProperty>(Property);
@@ -127,7 +127,7 @@ static FString GetVarTooltip(UBlueprint* InBlueprint, UClass* VarClass, FName Va
 	if (VarClass)
 	{
 	
-		if (FProperty* Property = FindField<FProperty>(VarClass, VarName))
+		if (FProperty* Property = FindFProperty<FProperty>(VarClass, VarName))
 		{
 			// discover if the variable property is a non blueprint user variable
 			UClass* SourceClass = Property->GetOwnerClass();
@@ -221,7 +221,7 @@ static void GetSubGraphIcon(FEdGraphSchemaAction_K2Graph const* const ActionIn,
 			{
 				// Need to see if this is a function overriding something in the parent, or 
 				;
-				if (UFunction* OverrideFunc = FindField<UFunction>(BlueprintIn->ParentClass, ActionIn->FuncName))
+				if (UFunction* OverrideFunc = FindUField<UFunction>(BlueprintIn->ParentClass, ActionIn->FuncName))
 				{
 					ToolTipOut = LOCTEXT("Override_Tooltip", "Override");
 				}
@@ -239,7 +239,7 @@ static void GetSubGraphIcon(FEdGraphSchemaAction_K2Graph const* const ActionIn,
 			FFormatNamedArguments Args;
 			Args.Add(TEXT("InterfaceName"), FText::FromName(ActionIn->FuncName));
 			ToolTipOut = FText::Format(LOCTEXT("FunctionFromInterface_Tooltip", "Function (from Interface '{InterfaceName}')"), Args);
-			if (UFunction* OverrideFunc = FindField<UFunction>(BlueprintIn->SkeletonGeneratedClass, ActionIn->FuncName))
+			if (UFunction* OverrideFunc = FindUField<UFunction>(BlueprintIn->SkeletonGeneratedClass, ActionIn->FuncName))
 			{
 				if (UEdGraphSchema_K2::FunctionCanBePlacedAsEvent(OverrideFunc))
 				{
@@ -262,7 +262,7 @@ static void GetSubGraphIcon(FEdGraphSchemaAction_K2Graph const* const ActionIn,
 				{
 					IconOut = FEditorStyle::GetBrush(TEXT("GraphEditor.Animation_16x"));
 				}
-				else if (UFunction* OverrideFunc = FindField<UFunction>(BlueprintIn->ParentClass, ActionIn->FuncName))
+				else if (UFunction* OverrideFunc = FindUField<UFunction>(BlueprintIn->ParentClass, ActionIn->FuncName))
 				{
 					const bool bIsPureFunction = OverrideFunc && OverrideFunc->HasAnyFunctionFlags(FUNC_BlueprintPure);
 					IconOut = FEditorStyle::GetBrush(bIsPureFunction ? TEXT("GraphEditor.OverridePureFunction_16x") : TEXT("GraphEditor.OverrideFunction_16x"));
@@ -270,7 +270,7 @@ static void GetSubGraphIcon(FEdGraphSchemaAction_K2Graph const* const ActionIn,
 				}
 				else
 				{
-					UFunction* Function = FindField<UFunction>(BlueprintIn->SkeletonGeneratedClass, ActionIn->FuncName);
+					UFunction* Function = FindUField<UFunction>(BlueprintIn->SkeletonGeneratedClass, ActionIn->FuncName);
 					const bool bIsPureFunction = Function && Function->HasAnyFunctionFlags(FUNC_BlueprintPure);
 
 					IconOut = FEditorStyle::GetBrush(bIsPureFunction ? TEXT("GraphEditor.PureFunction_16x") : TEXT("GraphEditor.Function_16x"));
@@ -1096,7 +1096,7 @@ void SBlueprintPaletteItem::Construct(const FArguments& InArgs, FCreateWidgetFor
 
 		if (FEdGraphSchemaAction_K2Graph* FuncGraphAction = (FEdGraphSchemaAction_K2Graph*)(GraphAction.Get()))
 		{
-			FunctionToCheck = FindField<UFunction>(Blueprint->SkeletonGeneratedClass, FuncGraphAction->FuncName);
+			FunctionToCheck = FindUField<UFunction>(Blueprint->SkeletonGeneratedClass, FuncGraphAction->FuncName);
 
 			// Handle override/interface functions
 			if(!FunctionToCheck)
diff --git a/Engine/Source/Editor/Kismet/Private/SBlueprintSubPalette.cpp b/Engine/Source/Editor/Kismet/Private/SBlueprintSubPalette.cpp
--- a/Engine/Source/Editor/Kismet/Private/SBlueprintSubPalette.cpp
+++ b/Engine/Source/Editor/Kismet/Private/SBlueprintSubPalette.cpp
@@ -84,7 +84,7 @@ static bool CanPaletteItemBePlaced(TSharedPtr<FEdGraphSchemaAction> DropActionIn
 			UClass const* const FuncOwner = CallFuncNode->FunctionReference.GetMemberParentClass(CallFuncNode->GetBlueprintClassFromNode());
 			check(FuncOwner != nullptr);
 
-			UFunction* const Function = FindField<UFunction>(FuncOwner, FuncName);
+			UFunction* const Function = FindUField<UFunction>(FuncOwner, FuncName);
 			UEdGraphSchema_K2 const* const K2Schema = Cast<UEdGraphSchema_K2 const>(GraphSchema);
 
 			if (Function == nullptr)
diff --git a/Engine/Source/Editor/Kismet/Private/SKismetInspector.cpp b/Engine/Source/Editor/Kismet/Private/SKismetInspector.cpp
--- a/Engine/Source/Editor/Kismet/Private/SKismetInspector.cpp
+++ b/Engine/Source/Editor/Kismet/Private/SKismetInspector.cpp
@@ -703,7 +703,7 @@ void SKismetInspector::UpdateFromObjects(const TArray<UObject*>& PropertyObjects
 
 				if (Object != EditableComponentTemplate)
 				{
-					if (FObjectProperty* ObjectProperty = FindField<FObjectProperty>(Object->GetClass(), EditableComponentTemplate->GetFName()))
+					if (FObjectProperty* ObjectProperty = FindFProperty<FObjectProperty>(Object->GetClass(), EditableComponentTemplate->GetFName()))
 					{
 						SelectedObjectProperties.Add(ObjectProperty);
 					}
diff --git a/Engine/Source/Editor/Kismet/Private/SMyBlueprint.cpp b/Engine/Source/Editor/Kismet/Private/SMyBlueprint.cpp
--- a/Engine/Source/Editor/Kismet/Private/SMyBlueprint.cpp
+++ b/Engine/Source/Editor/Kismet/Private/SMyBlueprint.cpp
@@ -1103,7 +1103,7 @@ void SMyBlueprint::GetLocalVariables(FGraphActionSort& SortList) const
 					Category = FText::GetEmpty();
 				}
 
-				UFunction* Func = FindField<UFunction>(GetBlueprintObj()->SkeletonGeneratedClass, TopLevelGraph->GetFName());
+				UFunction* Func = FindUField<UFunction>(GetBlueprintObj()->SkeletonGeneratedClass, TopLevelGraph->GetFName());
 				if (Func)
 				{
 					TSharedPtr<FEdGraphSchemaAction_K2LocalVar> NewVarAction = MakeShareable(new FEdGraphSchemaAction_K2LocalVar(Category, FText::FromName(Variable.VarName), FText::GetEmpty(), 0, NodeSectionID::LOCAL_VARIABLE));
@@ -2747,7 +2747,7 @@ struct FDeleteEntryHelper
 	{
 		check(NULL != Blueprint);
 
-		const FProperty* VariableProperty = FindField<FProperty>(Blueprint->SkeletonGeneratedClass, VarName);
+		const FProperty* VariableProperty = FindFProperty<FProperty>(Blueprint->SkeletonGeneratedClass, VarName);
 		const UClass* VarSourceClass = VariableProperty->GetOwnerChecked<const UClass>();
 		const bool bIsBlueprintVariable = (VarSourceClass == Blueprint->SkeletonGeneratedClass);
 		const int32 VarInfoIndex = FBlueprintEditorUtils::FindNewVariableIndex(Blueprint, VariableProperty->GetFName());
diff --git a/Engine/Source/Editor/Kismet/Private/SSCSEditor.cpp b/Engine/Source/Editor/Kismet/Private/SSCSEditor.cpp
--- a/Engine/Source/Editor/Kismet/Private/SSCSEditor.cpp
+++ b/Engine/Source/Editor/Kismet/Private/SSCSEditor.cpp
@@ -964,7 +964,7 @@ FString FSCSEditorTreeNodeComponentBase::GetDisplayString() const
 
 	bool const bHasValidVarName = (VariableName != NAME_None);
 	bool const bIsArrayVariable = bHasValidVarName && (VariableOwner != nullptr) && 
-		FindField<FArrayProperty>(VariableOwner, VariableName);
+		FindFProperty<FArrayProperty>(VariableOwner, VariableName);
 
 	// Only display SCS node variable names in the tree if they have not been autogenerated
 	if ((VariableName != NAME_None) && !bIsArrayVariable)
@@ -1309,7 +1309,7 @@ UActorComponent* FSCSEditorTreeNode::FindComponentInstanceInActor(const AActor*
 			if (VariableName != NAME_None)
 			{
 				UWorld* World = InActor->GetWorld();
-				FObjectPropertyBase* Property = FindField<FObjectPropertyBase>(InActor->GetClass(), VariableName);
+				FObjectPropertyBase* Property = FindFProperty<FObjectPropertyBase>(InActor->GetClass(), VariableName);
 				if (Property != NULL)
 				{
 					// Return the component instance that's stored in the property with the given variable name
@@ -4083,7 +4083,7 @@ void SSCSEditor::BuildMenuEventsSection(FMenuBuilder& Menu, UBlueprint* Blueprin
 				if( NodeIter->Component.IsValid() )
 				{
 					FName VariableName = NodeIter->VariableName;
-					FObjectProperty* VariableProperty = FindField<FObjectProperty>( Blueprint->SkeletonGeneratedClass, VariableName );
+					FObjectProperty* VariableProperty = FindFProperty<FObjectProperty>( Blueprint->SkeletonGeneratedClass, VariableName );
 
 					if( VariableProperty && FKismetEditorUtilities::FindBoundEventForComponent( Blueprint, EventName, VariableProperty->GetFName() ))
 					{
@@ -4138,7 +4138,7 @@ void SSCSEditor::CreateEventsForSelection(UBlueprint* Blueprint, FName EventName
 void SSCSEditor::ConstructEvent(UBlueprint* Blueprint, const FName EventName, const FComponentEventConstructionData EventData)
 {
 	// Find the corresponding variable property in the Blueprint
-	FObjectProperty* VariableProperty = FindField<FObjectProperty>(Blueprint->SkeletonGeneratedClass, EventData.VariableName );
+	FObjectProperty* VariableProperty = FindFProperty<FObjectProperty>(Blueprint->SkeletonGeneratedClass, EventData.VariableName );
 
 	if( VariableProperty )
 	{
@@ -4152,7 +4152,7 @@ void SSCSEditor::ConstructEvent(UBlueprint* Blueprint, const FName EventName, co
 void SSCSEditor::ViewEvent(UBlueprint* Blueprint, const FName EventName, const FComponentEventConstructionData EventData)
 {
 	// Find the corresponding variable property in the Blueprint
-	FObjectProperty* VariableProperty = FindField<FObjectProperty>(Blueprint->SkeletonGeneratedClass, EventData.VariableName );
+	FObjectProperty* VariableProperty = FindFProperty<FObjectProperty>(Blueprint->SkeletonGeneratedClass, EventData.VariableName );
 
 	if( VariableProperty )
 	{
@@ -4548,7 +4548,7 @@ void SSCSEditor::HighlightTreeNode(const USCS_Node* Node, FName Property)
 	if( Property != FName() )
 	{
 		UActorComponent* Component = TreeNode->GetComponentTemplate();
-		FProperty* CurrentProp = FindField<FProperty>(Component->GetClass(), Property);
+		FProperty* CurrentProp = FindFProperty<FProperty>(Component->GetClass(), Property);
 		FPropertyPath Path;
 		if( CurrentProp )
 		{
diff --git a/Engine/Source/Editor/KismetCompiler/Private/AnimBlueprintCompiler.cpp b/Engine/Source/Editor/KismetCompiler/Private/AnimBlueprintCompiler.cpp
--- a/Engine/Source/Editor/KismetCompiler/Private/AnimBlueprintCompiler.cpp
+++ b/Engine/Source/Editor/KismetCompiler/Private/AnimBlueprintCompiler.cpp
@@ -769,7 +769,7 @@ void FAnimBlueprintCompilerContext::ProcessUseCachedPose(UAnimGraphNode_UseCache
 	{
 		if (UAnimGraphNode_SaveCachedPose* AssociatedSaveNode = SaveCachedPoseNodes.FindRef(UseCachedPose->SaveCachedPoseNode->CacheName))
 		{
-			FStructProperty* LinkProperty = FindField<FStructProperty>(FAnimNode_UseCachedPose::StaticStruct(), TEXT("LinkToCachingNode"));
+			FStructProperty* LinkProperty = FindFProperty<FStructProperty>(FAnimNode_UseCachedPose::StaticStruct(), TEXT("LinkToCachingNode"));
 			check(LinkProperty);
 
 			FPoseLinkMappingRecord LinkRecord = FPoseLinkMappingRecord::MakeFromMember(UseCachedPose, AssociatedSaveNode, LinkProperty);
@@ -828,7 +828,7 @@ void FAnimBlueprintCompilerContext::ProcessCustomPropertyNode(UAnimGraphNode_Cus
 				if (!bGenerateLinkedAnimGraphVariables)
 				{
 					UClass* InstClass = CustomPropNode->GetTargetSkeletonClass();
-					if (FProperty* FoundProperty = FindField<FProperty>(InstClass, Pin->PinName))
+					if (FProperty* FoundProperty = FindFProperty<FProperty>(InstClass, Pin->PinName))
 					{
 						CustomPropNode->AddSourceTargetProperties(NewProperty->GetFName(), FoundProperty->GetFName());
 					}
@@ -1793,7 +1793,7 @@ void FAnimBlueprintCompilerContext::CopyTermDefaultsToDefaultObject(UObject* Def
 			{
 				if (RootStructProp->Struct->IsChildOf(FAnimNode_Base::StaticStruct()))
 				{
-					FStructProperty* ChildStructProp = FindField<FStructProperty>(NewAnimBlueprintClass, *RootStructProp->GetName());
+					FStructProperty* ChildStructProp = FindFProperty<FStructProperty>(NewAnimBlueprintClass, *RootStructProp->GetName());
 					check(ChildStructProp);
 					uint8* SourcePtr = RootStructProp->ContainerPtrToValuePtr<uint8>(RootDefaultObject);
 					uint8* DestPtr = ChildStructProp->ContainerPtrToValuePtr<uint8>(DefaultAnimInstance);
diff --git a/Engine/Source/Editor/KismetCompiler/Private/KismetCompiler.cpp b/Engine/Source/Editor/KismetCompiler/Private/KismetCompiler.cpp
--- a/Engine/Source/Editor/KismetCompiler/Private/KismetCompiler.cpp
+++ b/Engine/Source/Editor/KismetCompiler/Private/KismetCompiler.cpp
@@ -572,7 +572,7 @@ void FKismetCompilerContext::ValidateVariableNames()
 			else if (ParentClass->IsNative()) // the above case handles when the parent is a blueprint
 			{
 				FFieldVariant ExisingField = FindUFieldOrFProperty(ParentClass, *VarNameStr);
-				if (ExisingField.IsValid())
+				if (ExisingField)
 				{
 					UE_LOG(LogK2Compiler, Warning, TEXT("ValidateVariableNames name %s (used in %s) is already taken by %s")
 						, *VarNameStr, *Blueprint->GetPathName(), *ExisingField.GetPathName());
@@ -694,7 +694,7 @@ void FKismetCompilerContext::CreateClassVariablesFromBlueprint()
 			{
 				if(FMulticastDelegateProperty* AsDelegate = CastField<FMulticastDelegateProperty>(NewProperty))
 				{
-					AsDelegate->SignatureFunction = FindField<UFunction>(NewClass, *(Variable.VarName.ToString() + HEADER_GENERATED_DELEGATE_SIGNATURE_SUFFIX));
+					AsDelegate->SignatureFunction = FindUField<UFunction>(NewClass, *(Variable.VarName.ToString() + HEADER_GENERATED_DELEGATE_SIGNATURE_SUFFIX));
 					// Skeleton compilation phase may run when the delegate has been created but the function has not:
 					ensureAlways(AsDelegate->SignatureFunction || !bIsFullCompile);
 				}
@@ -1674,7 +1674,7 @@ void FKismetCompilerContext::PrecompileFunction(FKismetFunctionContext& Context,
 			);
 			return;
 		}
-		else if (NULL != FindField<FProperty>(NewClass, NewFunctionName))
+		else if (NULL != FindFProperty<FProperty>(NewClass, NewFunctionName))
 		{
 			MessageLog.Error(
 				*FText::Format(
@@ -1964,7 +1964,7 @@ void FKismetCompilerContext::PrecompileFunction(FKismetFunctionContext& Context,
 		{
 			Context.Function->FunctionFlags |= FUNC_Delegate;
 
-			if (FMulticastDelegateProperty* Property = FindField<FMulticastDelegateProperty>(NewClass, Context.DelegateSignatureName))
+			if (FMulticastDelegateProperty* Property = FindFProperty<FMulticastDelegateProperty>(NewClass, Context.DelegateSignatureName))
 			{
 				Property->SignatureFunction = Context.Function;
 			}
diff --git a/Engine/Source/Editor/KismetCompiler/Private/KismetCompilerMisc.cpp b/Engine/Source/Editor/KismetCompiler/Private/KismetCompilerMisc.cpp
--- a/Engine/Source/Editor/KismetCompiler/Private/KismetCompilerMisc.cpp
+++ b/Engine/Source/Editor/KismetCompiler/Private/KismetCompilerMisc.cpp
@@ -743,7 +743,7 @@ UEdGraphPin* FKismetCompilerUtilities::GenerateAssignmentNodes(class FKismetComp
 		{
 			if( OrgPin->LinkedTo.Num() == 0 )
 			{
-				FProperty* Property = FindField<FProperty>(ForClass, OrgPin->PinName);
+				FProperty* Property = FindFProperty<FProperty>(ForClass, OrgPin->PinName);
 				// NULL property indicates that this pin was part of the original node, not the 
 				// class we're assigning to:
 				if( !Property )
@@ -1290,7 +1290,7 @@ FFieldVariant FKismetCompilerUtilities::CheckPropertyNameOnScope(UStruct* Scope,
 
 	if (Scope && !Scope->IsA<UFunction>() && (UBlueprintGeneratedClass::GetUberGraphFrameName() != PropertyName))
 	{
-		if (FProperty* Field = FindField<FProperty>(Scope->GetSuperStruct(), *NameStr))
+		if (FProperty* Field = FindFProperty<FProperty>(Scope->GetSuperStruct(), *NameStr))
 		{
 			return Field;
 		}
diff --git a/Engine/Source/Editor/KismetCompiler/Private/UserDefinedStructureCompilerUtils.cpp b/Engine/Source/Editor/KismetCompiler/Private/UserDefinedStructureCompilerUtils.cpp
--- a/Engine/Source/Editor/KismetCompiler/Private/UserDefinedStructureCompilerUtils.cpp
+++ b/Engine/Source/Editor/KismetCompiler/Private/UserDefinedStructureCompilerUtils.cpp
@@ -211,7 +211,7 @@ struct FUserDefinedStructureCompilerInner
 			bool bIsNewVariable = false;
 			if (FStructureEditorUtils::FStructEditorManager::ActiveChange == FStructureEditorUtils::EStructureEditorChangeInfo::DefaultValueChanged)
 			{
-				VarProperty = FindField<FProperty>(Struct, VarDesc.VarName);
+				VarProperty = FindFProperty<FProperty>(Struct, VarDesc.VarName);
 				if (!ensureMsgf(VarProperty, TEXT("Could not find the expected property (%s); was the struct (%s) unexpectedly sanitized?"), *VarDesc.VarName.ToString(), *Struct->GetName()))
 				{
 					VarProperty = FKismetCompilerUtilities::CreatePropertyOnScope(Struct, VarDesc.VarName, VarType, NULL, CPF_None, Schema, MessageLog);
diff --git a/Engine/Source/Editor/MovieSceneTools/Private/TrackEditors/EventTrackEditor.cpp b/Engine/Source/Editor/MovieSceneTools/Private/TrackEditors/EventTrackEditor.cpp
--- a/Engine/Source/Editor/MovieSceneTools/Private/TrackEditors/EventTrackEditor.cpp
+++ b/Engine/Source/Editor/MovieSceneTools/Private/TrackEditors/EventTrackEditor.cpp
@@ -171,7 +171,7 @@ TSharedPtr<SWidget> FEventTrackEditor::BuildOutlinerEditWidget(const FGuid& Obje
 void FEventTrackEditor::BuildTrackContextMenu(FMenuBuilder& MenuBuilder, UMovieSceneTrack* Track)
 {
 	UMovieSceneEventTrack* EventTrack = CastChecked<UMovieSceneEventTrack>(Track);
-	FProperty* EventPositionProperty = FindField<FProperty>(Track->GetClass(), GET_MEMBER_NAME_STRING_CHECKED(UMovieSceneEventTrack, EventPosition));
+	FProperty* EventPositionProperty = FindFProperty<FProperty>(Track->GetClass(), GET_MEMBER_NAME_STRING_CHECKED(UMovieSceneEventTrack, EventPosition));
 
 	FGuid ObjectBinding;
 	EventTrack->GetTypedOuter<UMovieScene>()->FindTrackBinding(*EventTrack, ObjectBinding);
diff --git a/Engine/Source/Editor/Persona/Private/PersonaMeshDetails.cpp b/Engine/Source/Editor/Persona/Private/PersonaMeshDetails.cpp
--- a/Engine/Source/Editor/Persona/Private/PersonaMeshDetails.cpp
+++ b/Engine/Source/Editor/Persona/Private/PersonaMeshDetails.cpp
@@ -3940,7 +3940,7 @@ void FPersonaMeshDetails::OnMaterialArrayChanged(UMaterialInterface* NewMaterial
 		// Whether or not we made a transaction and need to end it
 		bool bMadeTransaction = false;
 
-		FProperty* MaterialProperty = FindField<FProperty>(USkeletalMesh::StaticClass(), "Materials");
+		FProperty* MaterialProperty = FindFProperty<FProperty>(USkeletalMesh::StaticClass(), "Materials");
 		check(MaterialProperty);
 		Mesh->PreEditChange(MaterialProperty);
 		check(Mesh->Materials.Num() > SlotIndex)
@@ -4118,7 +4118,7 @@ void FPersonaMeshDetails::OnMaterialNameCommitted(const FText& InValue, ETextCom
 	{
 		FScopedTransaction ScopeTransaction(LOCTEXT("PersonaMaterialSlotNameChanged", "Persona editor: Material slot name change"));
 
-		FProperty* ChangedProperty = FindField<FProperty>(USkeletalMesh::StaticClass(), "Materials");
+		FProperty* ChangedProperty = FindFProperty<FProperty>(USkeletalMesh::StaticClass(), "Materials");
 		check(ChangedProperty);
 		SkeletalMeshPtr->PreEditChange(ChangedProperty);
 
diff --git a/Engine/Source/Editor/Persona/Private/SAnimationSequenceBrowser.cpp b/Engine/Source/Editor/Persona/Private/SAnimationSequenceBrowser.cpp
--- a/Engine/Source/Editor/Persona/Private/SAnimationSequenceBrowser.cpp
+++ b/Engine/Source/Editor/Persona/Private/SAnimationSequenceBrowser.cpp
@@ -1250,7 +1250,7 @@ TSharedRef<SToolTip> SAnimationSequenceBrowser::CreateCustomAssetToolTip(FAssetD
 		{
 			// Check for DisplayName metadata
 			FText DisplayName;
-			if (FProperty* Field = FindField<FProperty>(AssetClass, TagPair.Key))
+			if (FProperty* Field = FindFProperty<FProperty>(AssetClass, TagPair.Key))
 			{
 				DisplayName = Field->GetDisplayNameText();
 			}
diff --git a/Engine/Source/Editor/PhysicsAssetEditor/Private/PhysicsAssetEditor.cpp b/Engine/Source/Editor/PhysicsAssetEditor/Private/PhysicsAssetEditor.cpp
--- a/Engine/Source/Editor/PhysicsAssetEditor/Private/PhysicsAssetEditor.cpp
+++ b/Engine/Source/Editor/PhysicsAssetEditor/Private/PhysicsAssetEditor.cpp
@@ -1380,7 +1380,7 @@ bool FPhysicsAssetEditor::ShouldFilterAssetBasedOnSkeleton( const FAssetData& As
 void FPhysicsAssetEditor::CreateOrConvertConstraint(EPhysicsAssetEditorConstraintType ConstraintType)
 {
 	//we have to manually call PostEditChange to ensure profiles are updated correctly
-	FProperty* DefaultInstanceProperty = FindField<FProperty>(UPhysicsConstraintTemplate::StaticClass(), GET_MEMBER_NAME_CHECKED(UPhysicsConstraintTemplate, DefaultInstance));
+	FProperty* DefaultInstanceProperty = FindFProperty<FProperty>(UPhysicsConstraintTemplate::StaticClass(), GET_MEMBER_NAME_CHECKED(UPhysicsConstraintTemplate, DefaultInstance));
 
 	const FScopedTransaction Transaction( LOCTEXT( "CreateConvertConstraint", "Create Or Convert Constraint" ) );
 
diff --git a/Engine/Source/Editor/PropertyEditor/Private/EditConditionContext.cpp b/Engine/Source/Editor/PropertyEditor/Private/EditConditionContext.cpp
--- a/Engine/Source/Editor/PropertyEditor/Private/EditConditionContext.cpp
+++ b/Engine/Source/Editor/PropertyEditor/Private/EditConditionContext.cpp
@@ -39,7 +39,7 @@ const FBoolProperty* FEditConditionContext::GetSingleBoolProperty(const TSharedP
 				return nullptr;
 			}
 
-			const FProperty* Field = FindField<FProperty>(Property->GetOwnerStruct(), *PropertyToken->PropertyName);
+			const FProperty* Field = FindFProperty<FProperty>(Property->GetOwnerStruct(), *PropertyToken->PropertyName);
 			BoolProperty = CastField<FBoolProperty>(Field);
 			
 			// not a bool
@@ -63,7 +63,7 @@ T* FindTypedField(const TWeakPtr<FPropertyNode>& PropertyNode, const FString& Pr
 		TSharedPtr<FPropertyNode> PinnedNode = PropertyNode.Pin();
 		const FProperty* Property = PinnedNode->GetProperty();
 
-		FProperty* Field = FindField<FProperty>(Property->GetOwnerStruct(), *PropertyName);
+		FProperty* Field = FindFProperty<FProperty>(Property->GetOwnerStruct(), *PropertyName);
 		if (Field == nullptr)
 		{
 			if (!AlreadyLogged.Find(Field))
diff --git a/Engine/Source/Editor/PropertyEditor/Private/PropertyCustomizationHelpers.cpp b/Engine/Source/Editor/PropertyEditor/Private/PropertyCustomizationHelpers.cpp
--- a/Engine/Source/Editor/PropertyEditor/Private/PropertyCustomizationHelpers.cpp
+++ b/Engine/Source/Editor/PropertyEditor/Private/PropertyCustomizationHelpers.cpp
@@ -408,7 +408,7 @@ namespace PropertyCustomizationHelpers
 			if ( ConditionPropertyName.Len() > 0 && !ConditionPropertyName.Contains(TEXT(".")) )
 			{
 				UStruct* Scope = InProperty->GetOwnerStruct();
-				EditConditionProperty = FindField<FBoolProperty>(Scope, *ConditionPropertyName);
+				EditConditionProperty = FindFProperty<FBoolProperty>(Scope, *ConditionPropertyName);
 			}
 		}
 
diff --git a/Engine/Source/Editor/SceneOutliner/Private/ComponentTreeItem.cpp b/Engine/Source/Editor/SceneOutliner/Private/ComponentTreeItem.cpp
--- a/Engine/Source/Editor/SceneOutliner/Private/ComponentTreeItem.cpp
+++ b/Engine/Source/Editor/SceneOutliner/Private/ComponentTreeItem.cpp
@@ -37,7 +37,7 @@ FComponentTreeItem::FComponentTreeItem(UActorComponent* InComponent)
 	bExistsInCurrentWorldAndPIE = GEditor->ObjectsThatExistInEditorWorld.Get(OwningActor);
 
 	const FName VariableName = FComponentEditorUtils::FindVariableNameGivenComponentInstance(InComponent);
-	const bool bIsArrayVariable = !VariableName.IsNone() && InComponent->GetOwner() != nullptr && FindField<FArrayProperty>(InComponent->GetOwner()->GetClass(), VariableName);
+	const bool bIsArrayVariable = !VariableName.IsNone() && InComponent->GetOwner() != nullptr && FindFProperty<FArrayProperty>(InComponent->GetOwner()->GetClass(), VariableName);
 
 	if (!VariableName.IsNone() && !bIsArrayVariable)
 	{
diff --git a/Engine/Source/Editor/StaticMeshEditor/Private/StaticMeshEditorTools.cpp b/Engine/Source/Editor/StaticMeshEditor/Private/StaticMeshEditorTools.cpp
--- a/Engine/Source/Editor/StaticMeshEditor/Private/StaticMeshEditorTools.cpp
+++ b/Engine/Source/Editor/StaticMeshEditor/Private/StaticMeshEditorTools.cpp
@@ -2819,7 +2819,7 @@ void FMeshMaterialsLayout::OnMaterialNameCommitted(const FText& InValue, ETextCo
 		FScopedTransaction ScopeTransaction(LOCTEXT("StaticMeshEditorMaterialSlotNameChanged", "Staticmesh editor: Material slot name change"));
 
 		FProperty* ChangedProperty = NULL;
-		ChangedProperty = FindField<FProperty>(UStaticMesh::StaticClass(), "StaticMaterials");
+		ChangedProperty = FindFProperty<FProperty>(UStaticMesh::StaticClass(), "StaticMaterials");
 		check(ChangedProperty);
 		StaticMesh.PreEditChange(ChangedProperty);
 
diff --git a/Engine/Source/Editor/StaticMeshEditor/Private/StaticMeshEditorViewportClient.cpp b/Engine/Source/Editor/StaticMeshEditor/Private/StaticMeshEditorViewportClient.cpp
--- a/Engine/Source/Editor/StaticMeshEditor/Private/StaticMeshEditorViewportClient.cpp
+++ b/Engine/Source/Editor/StaticMeshEditor/Private/StaticMeshEditorViewportClient.cpp
@@ -179,7 +179,7 @@ bool FStaticMeshEditorViewportClient::InputWidgetDelta( FViewport* InViewport, E
 				const FWidget::EWidgetMode MoveMode = GetWidgetMode();
 				if(MoveMode == FWidget::WM_Rotate)
 				{
-					ChangedProperty = FindField<FProperty>( UStaticMeshSocket::StaticClass(), "RelativeRotation" );
+					ChangedProperty = FindFProperty<FProperty>( UStaticMeshSocket::StaticClass(), "RelativeRotation" );
 					SelectedSocket->PreEditChange(ChangedProperty);
 
 					FRotator CurrentRot = SelectedSocket->RelativeRotation;
@@ -198,7 +198,7 @@ bool FStaticMeshEditorViewportClient::InputWidgetDelta( FViewport* InViewport, E
 				}
 				else if(MoveMode == FWidget::WM_Translate)
 				{
-					ChangedProperty = FindField<FProperty>( UStaticMeshSocket::StaticClass(), "RelativeLocation" );
+					ChangedProperty = FindFProperty<FProperty>( UStaticMeshSocket::StaticClass(), "RelativeLocation" );
 					SelectedSocket->PreEditChange(ChangedProperty);
 
 					//FRotationMatrix SocketRotTM( SelectedSocket->RelativeRotation );
diff --git a/Engine/Source/Editor/TranslationEditor/Private/TranslationEditor.cpp b/Engine/Source/Editor/TranslationEditor/Private/TranslationEditor.cpp
--- a/Engine/Source/Editor/TranslationEditor/Private/TranslationEditor.cpp
+++ b/Engine/Source/Editor/TranslationEditor/Private/TranslationEditor.cpp
@@ -251,8 +251,8 @@ TSharedRef<SDockTab> FTranslationEditor::SpawnTab_Untranslated( const FSpawnTabA
 
 	FPropertyEditorModule& PropertyEditorModule = FModuleManager::GetModuleChecked<FPropertyEditorModule>( "PropertyEditor" );
 
-	FProperty* SourceProperty = FindField<FProperty>( UTranslationUnit::StaticClass(), "Source");
-	FProperty* TranslationProperty = FindField<FProperty>( UTranslationUnit::StaticClass(), "Translation");
+	FProperty* SourceProperty = FindFProperty<FProperty>( UTranslationUnit::StaticClass(), "Source");
+	FProperty* TranslationProperty = FindFProperty<FProperty>( UTranslationUnit::StaticClass(), "Translation");
 
 	// create empty property table
 	UntranslatedPropertyTable = PropertyEditorModule.CreatePropertyTable();
@@ -307,8 +307,8 @@ TSharedRef<SDockTab> FTranslationEditor::SpawnTab_Review( const FSpawnTabArgs& A
 {
 	check( Args.GetTabId().TabType == ReviewTabId );
 
-	FProperty* SourceProperty = FindField<FProperty>( UTranslationUnit::StaticClass(), "Source");
-	FProperty* TranslationProperty = FindField<FProperty>( UTranslationUnit::StaticClass(), "Translation");
+	FProperty* SourceProperty = FindFProperty<FProperty>( UTranslationUnit::StaticClass(), "Source");
+	FProperty* TranslationProperty = FindFProperty<FProperty>( UTranslationUnit::StaticClass(), "Translation");
 
 	FPropertyEditorModule& PropertyEditorModule = FModuleManager::GetModuleChecked<FPropertyEditorModule>( "PropertyEditor" );
 
@@ -330,9 +330,9 @@ TSharedRef<SDockTab> FTranslationEditor::SpawnTab_Review( const FSpawnTabArgs& A
 	ReviewPropertyTable->SetObjects((TArray<UObject*>&)DataManager->GetReviewArray());
 
 	// Add the columns we want to display
-	ReviewPropertyTable->AddColumn((TWeakFieldPtr<FProperty>)FindField<FProperty>( UTranslationUnit::StaticClass(), "Source"));
-	ReviewPropertyTable->AddColumn((TWeakFieldPtr<FProperty>)FindField<FProperty>( UTranslationUnit::StaticClass(), "Translation"));
-	ReviewPropertyTable->AddColumn((TWeakFieldPtr<FProperty>)FindField<FProperty>( UTranslationUnit::StaticClass(), "HasBeenReviewed"));
+	ReviewPropertyTable->AddColumn((TWeakFieldPtr<FProperty>)FindFProperty<FProperty>( UTranslationUnit::StaticClass(), "Source"));
+	ReviewPropertyTable->AddColumn((TWeakFieldPtr<FProperty>)FindFProperty<FProperty>( UTranslationUnit::StaticClass(), "Translation"));
+	ReviewPropertyTable->AddColumn((TWeakFieldPtr<FProperty>)FindFProperty<FProperty>( UTranslationUnit::StaticClass(), "HasBeenReviewed"));
 
 	TArray<TSharedRef<IPropertyTableColumn>> Columns = ReviewPropertyTable->GetColumns();
 	for (TSharedRef<IPropertyTableColumn> Column : Columns)
@@ -372,8 +372,8 @@ TSharedRef<SDockTab> FTranslationEditor::SpawnTab_Completed( const FSpawnTabArgs
 {
 	check( Args.GetTabId().TabType == CompletedTabId );
 
-	FProperty* SourceProperty = FindField<FProperty>( UTranslationUnit::StaticClass(), "Source");
-	FProperty* TranslationProperty = FindField<FProperty>( UTranslationUnit::StaticClass(), "Translation");
+	FProperty* SourceProperty = FindFProperty<FProperty>( UTranslationUnit::StaticClass(), "Source");
+	FProperty* TranslationProperty = FindFProperty<FProperty>( UTranslationUnit::StaticClass(), "Translation");
 
 	FPropertyEditorModule& PropertyEditorModule = FModuleManager::GetModuleChecked<FPropertyEditorModule>( "PropertyEditor" );
 
@@ -395,8 +395,8 @@ TSharedRef<SDockTab> FTranslationEditor::SpawnTab_Completed( const FSpawnTabArgs
 	CompletedPropertyTable->SetObjects((TArray<UObject*>&)DataManager->GetCompleteArray());
 
 	// Add the columns we want to display
-	CompletedPropertyTable->AddColumn((TWeakFieldPtr<FProperty>)FindField<FProperty>( UTranslationUnit::StaticClass(), "Source"));
-	CompletedPropertyTable->AddColumn((TWeakFieldPtr<FProperty>)FindField<FProperty>( UTranslationUnit::StaticClass(), "Translation"));
+	CompletedPropertyTable->AddColumn((TWeakFieldPtr<FProperty>)FindFProperty<FProperty>( UTranslationUnit::StaticClass(), "Source"));
+	CompletedPropertyTable->AddColumn((TWeakFieldPtr<FProperty>)FindFProperty<FProperty>( UTranslationUnit::StaticClass(), "Translation"));
 
 	// Freeze columns, don't want user to remove them
 	TArray<TSharedRef<IPropertyTableColumn>> Columns = CompletedPropertyTable->GetColumns();
@@ -430,8 +430,8 @@ TSharedRef<SDockTab> FTranslationEditor::SpawnTab_Search(const FSpawnTabArgs& Ar
 {
 	check(Args.GetTabId().TabType == SearchTabId);
 
-	FProperty* SourceProperty = FindField<FProperty>(UTranslationUnit::StaticClass(), "Source");
-	FProperty* TranslationProperty = FindField<FProperty>(UTranslationUnit::StaticClass(), "Translation");
+	FProperty* SourceProperty = FindFProperty<FProperty>(UTranslationUnit::StaticClass(), "Source");
+	FProperty* TranslationProperty = FindFProperty<FProperty>(UTranslationUnit::StaticClass(), "Translation");
 
 	FPropertyEditorModule& PropertyEditorModule = FModuleManager::GetModuleChecked<FPropertyEditorModule>("PropertyEditor");
 
@@ -453,8 +453,8 @@ TSharedRef<SDockTab> FTranslationEditor::SpawnTab_Search(const FSpawnTabArgs& Ar
 	SearchPropertyTable->SetObjects((TArray<UObject*>&)DataManager->GetSearchResultsArray());
 
 	// Add the columns we want to display
-	SearchPropertyTable->AddColumn((TWeakFieldPtr<FProperty>)FindField<FProperty>(UTranslationUnit::StaticClass(), "Source"));
-	SearchPropertyTable->AddColumn((TWeakFieldPtr<FProperty>)FindField<FProperty>(UTranslationUnit::StaticClass(), "Translation"));
+	SearchPropertyTable->AddColumn((TWeakFieldPtr<FProperty>)FindFProperty<FProperty>(UTranslationUnit::StaticClass(), "Source"));
+	SearchPropertyTable->AddColumn((TWeakFieldPtr<FProperty>)FindFProperty<FProperty>(UTranslationUnit::StaticClass(), "Translation"));
 
 	// Freeze columns, don't want user to remove them
 	TArray<TSharedRef<IPropertyTableColumn>> Columns = SearchPropertyTable->GetColumns();
@@ -508,9 +508,9 @@ TSharedRef<SDockTab> FTranslationEditor::SpawnTab_ChangedOnImport(const FSpawnTa
 {
 	check(Args.GetTabId().TabType == ChangedOnImportTabId);
 
-	FProperty* SourceProperty = FindField<FProperty>(UTranslationUnit::StaticClass(), "Source");
-	FProperty* TranslationBeforeImportProperty = FindField<FProperty>(UTranslationUnit::StaticClass(), "TranslationBeforeImport");
-	FProperty* TranslationProperty = FindField<FProperty>(UTranslationUnit::StaticClass(), "Translation");
+	FProperty* SourceProperty = FindFProperty<FProperty>(UTranslationUnit::StaticClass(), "Source");
+	FProperty* TranslationBeforeImportProperty = FindFProperty<FProperty>(UTranslationUnit::StaticClass(), "TranslationBeforeImport");
+	FProperty* TranslationProperty = FindFProperty<FProperty>(UTranslationUnit::StaticClass(), "Translation");
 
 	FPropertyEditorModule& PropertyEditorModule = FModuleManager::GetModuleChecked<FPropertyEditorModule>("PropertyEditor");
 
@@ -532,9 +532,9 @@ TSharedRef<SDockTab> FTranslationEditor::SpawnTab_ChangedOnImport(const FSpawnTa
 	ChangedOnImportPropertyTable->SetObjects((TArray<UObject*>&)DataManager->GetSearchResultsArray());
 
 	// Add the columns we want to display
-	ChangedOnImportPropertyTable->AddColumn((TWeakFieldPtr<FProperty>)FindField<FProperty>(UTranslationUnit::StaticClass(), "Source"));
-	ChangedOnImportPropertyTable->AddColumn((TWeakFieldPtr<FProperty>)FindField<FProperty>(UTranslationUnit::StaticClass(), "TranslationBeforeImport"));
-	ChangedOnImportPropertyTable->AddColumn((TWeakFieldPtr<FProperty>)FindField<FProperty>(UTranslationUnit::StaticClass(), "Translation"));
+	ChangedOnImportPropertyTable->AddColumn((TWeakFieldPtr<FProperty>)FindFProperty<FProperty>(UTranslationUnit::StaticClass(), "Source"));
+	ChangedOnImportPropertyTable->AddColumn((TWeakFieldPtr<FProperty>)FindFProperty<FProperty>(UTranslationUnit::StaticClass(), "TranslationBeforeImport"));
+	ChangedOnImportPropertyTable->AddColumn((TWeakFieldPtr<FProperty>)FindFProperty<FProperty>(UTranslationUnit::StaticClass(), "Translation"));
 
 	// Freeze columns, don't want user to remove them
 	TArray<TSharedRef<IPropertyTableColumn>> Columns = ChangedOnImportPropertyTable->GetColumns();
@@ -612,7 +612,7 @@ TSharedRef<SDockTab> FTranslationEditor::SpawnTab_Context( const FSpawnTabArgs&
 	}
 
 	// Build the Path to the data we want to show
-	FProperty* ContextProp = FindField<FProperty>( UTranslationUnit::StaticClass(), "Contexts" );
+	FProperty* ContextProp = FindFProperty<FProperty>( UTranslationUnit::StaticClass(), "Contexts" );
 	FPropertyInfo ContextPropInfo;
 	ContextPropInfo.Property = ContextProp;
 	ContextPropInfo.ArrayIndex = INDEX_NONE;
@@ -621,8 +621,8 @@ TSharedRef<SDockTab> FTranslationEditor::SpawnTab_Context( const FSpawnTabArgs&
 	ContextPropertyTable->SetRootPath(Path);
 
 	// Add the columns we want to display
-	ContextPropertyTable->AddColumn((TWeakFieldPtr<FProperty>)FindField<FProperty>( FTranslationContextInfo::StaticStruct(), "Key"));
-	ContextPropertyTable->AddColumn((TWeakFieldPtr<FProperty>)FindField<FProperty>( FTranslationContextInfo::StaticStruct(), "Context"));
+	ContextPropertyTable->AddColumn((TWeakFieldPtr<FProperty>)FindFProperty<FProperty>( FTranslationContextInfo::StaticStruct(), "Key"));
+	ContextPropertyTable->AddColumn((TWeakFieldPtr<FProperty>)FindFProperty<FProperty>( FTranslationContextInfo::StaticStruct(), "Context"));
 
 	// Freeze columns, don't want user to remove them
 	TArray<TSharedRef<IPropertyTableColumn>> Columns = ContextPropertyTable->GetColumns();
@@ -674,8 +674,8 @@ TSharedRef<SDockTab> FTranslationEditor::SpawnTab_History(const FSpawnTabArgs& A
 
 	FPropertyEditorModule& PropertyEditorModule = FModuleManager::GetModuleChecked<FPropertyEditorModule>("PropertyEditor");
 
-	FProperty* SourceProperty = FindField<FProperty>(FTranslationChange::StaticStruct(), "Source");
-	FProperty* TranslationProperty = FindField<FProperty>(FTranslationChange::StaticStruct(), "Translation");
+	FProperty* SourceProperty = FindFProperty<FProperty>(FTranslationChange::StaticStruct(), "Source");
+	FProperty* TranslationProperty = FindFProperty<FProperty>(FTranslationChange::StaticStruct(), "Translation");
 
 	// create empty property table
 	HistoryPropertyTable = PropertyEditorModule.CreatePropertyTable();
@@ -700,14 +700,14 @@ TSharedRef<SDockTab> FTranslationEditor::SpawnTab_History(const FSpawnTabArgs& A
 
 	// Build the Path to the data we want to show
 	TSharedRef<FPropertyPath> Path = FPropertyPath::CreateEmpty();
-	FArrayProperty* ContextsProp = FindField<FArrayProperty>(UTranslationUnit::StaticClass(), "Contexts");
+	FArrayProperty* ContextsProp = FindFProperty<FArrayProperty>(UTranslationUnit::StaticClass(), "Contexts");
 	Path = Path->ExtendPath(FPropertyPath::Create(ContextsProp));
 	FPropertyInfo ContextsPropInfo;
 	ContextsPropInfo.Property = ContextsProp->Inner;
 	ContextsPropInfo.ArrayIndex = 0;
 	Path = Path->ExtendPath(ContextsPropInfo);
 
-	FProperty* ChangesProp = FindField<FProperty>(FTranslationContextInfo::StaticStruct(), "Changes");
+	FProperty* ChangesProp = FindFProperty<FProperty>(FTranslationContextInfo::StaticStruct(), "Changes");
 	FPropertyInfo ChangesPropInfo;
 	ChangesPropInfo.Property = ChangesProp;
 	ChangesPropInfo.ArrayIndex = INDEX_NONE;
@@ -715,8 +715,8 @@ TSharedRef<SDockTab> FTranslationEditor::SpawnTab_History(const FSpawnTabArgs& A
 	HistoryPropertyTable->SetRootPath(Path);
 
 	// Add the columns we want to display
-	HistoryPropertyTable->AddColumn((TWeakFieldPtr<FProperty>)FindField<FProperty>(FTranslationChange::StaticStruct(), "Version"));
-	HistoryPropertyTable->AddColumn((TWeakFieldPtr<FProperty>)FindField<FProperty>(FTranslationChange::StaticStruct(), "DateAndTime"));
+	HistoryPropertyTable->AddColumn((TWeakFieldPtr<FProperty>)FindFProperty<FProperty>(FTranslationChange::StaticStruct(), "Version"));
+	HistoryPropertyTable->AddColumn((TWeakFieldPtr<FProperty>)FindFProperty<FProperty>(FTranslationChange::StaticStruct(), "DateAndTime"));
 	HistoryPropertyTable->AddColumn((TWeakFieldPtr<FProperty>)SourceProperty);
 	HistoryPropertyTable->AddColumn((TWeakFieldPtr<FProperty>)TranslationProperty);
 
@@ -994,7 +994,7 @@ void FTranslationEditor::UpdateTranslationUnitSelection(TSet<TSharedRef<IPropert
 					NamespaceTextBlock->SetText(FText::Format(LOCTEXT("TranslationNamespace", "Namespace: {0}"), FText::FromString(SelectedTranslationUnit->Namespace)));
 
 					// Add the ContextPropertyTable-specific path
-					FArrayProperty* ContextArrayProp = FindField<FArrayProperty>(UTranslationUnit::StaticClass(), "Contexts");
+					FArrayProperty* ContextArrayProp = FindFProperty<FArrayProperty>(UTranslationUnit::StaticClass(), "Contexts");
 					FPropertyInfo ContextArrayPropInfo;
 					ContextArrayPropInfo.Property = ContextArrayProp;
 					ContextArrayPropInfo.ArrayIndex = INDEX_NONE;
@@ -1010,8 +1010,8 @@ void FTranslationEditor::UpdateTranslationUnitSelection(TSet<TSharedRef<IPropert
 						ContextPropertyTable->SetRootPath(ContextPath);
 
 						// Need to re-add the columns we want to display
-						ContextPropertyTable->AddColumn((TWeakFieldPtr<FProperty>)FindField<FProperty>(FTranslationContextInfo::StaticStruct(), "Key"));
-						ContextPropertyTable->AddColumn((TWeakFieldPtr<FProperty>)FindField<FProperty>(FTranslationContextInfo::StaticStruct(), "Context"));
+						ContextPropertyTable->AddColumn((TWeakFieldPtr<FProperty>)FindFProperty<FProperty>(FTranslationContextInfo::StaticStruct(), "Key"));
+						ContextPropertyTable->AddColumn((TWeakFieldPtr<FProperty>)FindFProperty<FProperty>(FTranslationContextInfo::StaticStruct(), "Context"));
 
 						TArray<TSharedRef<IPropertyTableColumn>> Columns = ContextPropertyTable->GetColumns();
 						for (TSharedRef<IPropertyTableColumn> Column : Columns)
@@ -1085,12 +1085,12 @@ void FTranslationEditor::UpdateContextSelection()
 
 						// Add the HistoryPropertyTable-specific path
 						TSharedRef<FPropertyPath> HistoryPath = ContextPropertyTable->GetRootPath();
-						FArrayProperty* ContextArrayProp = FindField<FArrayProperty>(UTranslationUnit::StaticClass(), "Contexts");
+						FArrayProperty* ContextArrayProp = FindFProperty<FArrayProperty>(UTranslationUnit::StaticClass(), "Contexts");
 						FPropertyInfo ContextPropInfo;
 						ContextPropInfo.Property = ContextArrayProp->Inner;
 						ContextPropInfo.ArrayIndex = PartialPath->GetLeafMostProperty().ArrayIndex;
 						HistoryPath = HistoryPath->ExtendPath(ContextPropInfo);
-						FArrayProperty* ChangesProp = FindField<FArrayProperty>(FTranslationContextInfo::StaticStruct(), "Changes");
+						FArrayProperty* ChangesProp = FindFProperty<FArrayProperty>(FTranslationContextInfo::StaticStruct(), "Changes");
 						FPropertyInfo ChangesPropInfo;
 						ChangesPropInfo.Property = ChangesProp;
 						ChangesPropInfo.ArrayIndex = INDEX_NONE;
@@ -1103,10 +1103,10 @@ void FTranslationEditor::UpdateContextSelection()
 							HistoryPropertyTable->SetRootPath(HistoryPath);
 
 							// Need to re-add the columns we want to display
-							HistoryPropertyTable->AddColumn((TWeakFieldPtr<FProperty>)FindField<FProperty>(FTranslationChange::StaticStruct(), "Version"));
-							HistoryPropertyTable->AddColumn((TWeakFieldPtr<FProperty>)FindField<FProperty>(FTranslationChange::StaticStruct(), "DateAndTime"));
-							HistoryPropertyTable->AddColumn((TWeakFieldPtr<FProperty>)FindField<FProperty>(FTranslationChange::StaticStruct(), "Source"));
-							HistoryPropertyTable->AddColumn((TWeakFieldPtr<FProperty>)FindField<FProperty>(FTranslationChange::StaticStruct(), "Translation"));
+							HistoryPropertyTable->AddColumn((TWeakFieldPtr<FProperty>)FindFProperty<FProperty>(FTranslationChange::StaticStruct(), "Version"));
+							HistoryPropertyTable->AddColumn((TWeakFieldPtr<FProperty>)FindFProperty<FProperty>(FTranslationChange::StaticStruct(), "DateAndTime"));
+							HistoryPropertyTable->AddColumn((TWeakFieldPtr<FProperty>)FindFProperty<FProperty>(FTranslationChange::StaticStruct(), "Source"));
+							HistoryPropertyTable->AddColumn((TWeakFieldPtr<FProperty>)FindFProperty<FProperty>(FTranslationChange::StaticStruct(), "Translation"));
 
 							TArray<TSharedRef<IPropertyTableColumn>> Columns = HistoryPropertyTable->GetColumns();
 							for (TSharedRef<IPropertyTableColumn> Column : Columns)
@@ -1342,9 +1342,9 @@ void FTranslationEditor::ImportFromPoFile(FString FileToImport)
 			TabManager->TryInvokeTab(ChangedOnImportTabId);
 			ChangedOnImportPropertyTable->SetObjects((TArray<UObject*>&)DataManager->GetChangedOnImportArray());
 			// Need to re-add the columns we want to display
-			ChangedOnImportPropertyTable->AddColumn((TWeakFieldPtr<FProperty>)FindField<FProperty>(UTranslationUnit::StaticClass(), "Source"));
-			ChangedOnImportPropertyTable->AddColumn((TWeakFieldPtr<FProperty>)FindField<FProperty>(UTranslationUnit::StaticClass(), "TranslationBeforeImport"));
-			ChangedOnImportPropertyTable->AddColumn((TWeakFieldPtr<FProperty>)FindField<FProperty>(UTranslationUnit::StaticClass(), "Translation"));
+			ChangedOnImportPropertyTable->AddColumn((TWeakFieldPtr<FProperty>)FindFProperty<FProperty>(UTranslationUnit::StaticClass(), "Source"));
+			ChangedOnImportPropertyTable->AddColumn((TWeakFieldPtr<FProperty>)FindFProperty<FProperty>(UTranslationUnit::StaticClass(), "TranslationBeforeImport"));
+			ChangedOnImportPropertyTable->AddColumn((TWeakFieldPtr<FProperty>)FindFProperty<FProperty>(UTranslationUnit::StaticClass(), "Translation"));
 		}
 	}
 	else
@@ -1382,8 +1382,8 @@ void FTranslationEditor::OnFilterTextCommitted(const FText& InFilterText, ETextC
 				SearchPropertyTable->SetObjects((TArray<UObject*>&)DataManager->GetSearchResultsArray());
 
 				// Need to re-add the columns we want to display
-				SearchPropertyTable->AddColumn((TWeakFieldPtr<FProperty>)FindField<FProperty>(UTranslationUnit::StaticClass(), "Source"));
-				SearchPropertyTable->AddColumn((TWeakFieldPtr<FProperty>)FindField<FProperty>(UTranslationUnit::StaticClass(), "Translation"));
+				SearchPropertyTable->AddColumn((TWeakFieldPtr<FProperty>)FindFProperty<FProperty>(UTranslationUnit::StaticClass(), "Source"));
+				SearchPropertyTable->AddColumn((TWeakFieldPtr<FProperty>)FindFProperty<FProperty>(UTranslationUnit::StaticClass(), "Translation"));
 
 				TArray<TSharedRef<IPropertyTableColumn>> Columns = SearchPropertyTable->GetColumns();
 				for (TSharedRef<IPropertyTableColumn> Column : Columns)
@@ -1411,8 +1411,8 @@ FReply FTranslationEditor::OnGetHistoryButtonClicked()
 		UntranslatedPropertyTable->SetObjects((TArray<UObject*>&)DataManager->GetUntranslatedArray());
 
 		// Need to re-add the columns we want to display
-		UntranslatedPropertyTable->AddColumn((TWeakFieldPtr<FProperty>)FindField<FProperty>(UTranslationUnit::StaticClass(), "Source"));
-		UntranslatedPropertyTable->AddColumn((TWeakFieldPtr<FProperty>)FindField<FProperty>(UTranslationUnit::StaticClass(), "Translation"));
+		UntranslatedPropertyTable->AddColumn((TWeakFieldPtr<FProperty>)FindFProperty<FProperty>(UTranslationUnit::StaticClass(), "Source"));
+		UntranslatedPropertyTable->AddColumn((TWeakFieldPtr<FProperty>)FindFProperty<FProperty>(UTranslationUnit::StaticClass(), "Translation"));
 
 		TArray<TSharedRef<IPropertyTableColumn>> Columns = UntranslatedPropertyTable->GetColumns();
 		for (TSharedRef<IPropertyTableColumn> Column : Columns)
@@ -1426,9 +1426,9 @@ FReply FTranslationEditor::OnGetHistoryButtonClicked()
 		ReviewPropertyTable->SetObjects((TArray<UObject*>&)DataManager->GetReviewArray());
 
 		// Need to re-add the columns we want to display
-		ReviewPropertyTable->AddColumn((TWeakFieldPtr<FProperty>)FindField<FProperty>(UTranslationUnit::StaticClass(), "Source"));
-		ReviewPropertyTable->AddColumn((TWeakFieldPtr<FProperty>)FindField<FProperty>(UTranslationUnit::StaticClass(), "Translation"));
-		ReviewPropertyTable->AddColumn((TWeakFieldPtr<FProperty>)FindField<FProperty>(UTranslationUnit::StaticClass(), "HasBeenReviewed"));
+		ReviewPropertyTable->AddColumn((TWeakFieldPtr<FProperty>)FindFProperty<FProperty>(UTranslationUnit::StaticClass(), "Source"));
+		ReviewPropertyTable->AddColumn((TWeakFieldPtr<FProperty>)FindFProperty<FProperty>(UTranslationUnit::StaticClass(), "Translation"));
+		ReviewPropertyTable->AddColumn((TWeakFieldPtr<FProperty>)FindFProperty<FProperty>(UTranslationUnit::StaticClass(), "HasBeenReviewed"));
 
 		TArray<TSharedRef<IPropertyTableColumn>> Columns = ReviewPropertyTable->GetColumns();
 		for (TSharedRef<IPropertyTableColumn> Column : Columns)
diff --git a/Engine/Source/Editor/UMGEditor/Private/Customizations/UMGDetailCustomizations.cpp b/Engine/Source/Editor/UMGEditor/Private/Customizations/UMGDetailCustomizations.cpp
--- a/Engine/Source/Editor/UMGEditor/Private/Customizations/UMGDetailCustomizations.cpp
+++ b/Engine/Source/Editor/UMGEditor/Private/Customizations/UMGDetailCustomizations.cpp
@@ -167,7 +167,7 @@ FReply FBlueprintWidgetCustomization::HandleAddOrViewEventForVariable(const FNam
 	UBlueprint* BlueprintObj = Blueprint;
 
 	// Find the corresponding variable property in the Blueprint
-	FObjectProperty* VariableProperty = FindField<FObjectProperty>(BlueprintObj->SkeletonGeneratedClass, PropertyName);
+	FObjectProperty* VariableProperty = FindFProperty<FObjectProperty>(BlueprintObj->SkeletonGeneratedClass, PropertyName);
 
 	if (VariableProperty)
 	{
@@ -218,7 +218,7 @@ void FBlueprintWidgetCustomization::CreateMulticastEventCustomization(IDetailLay
 		PropertyTooltip = FText::FromString(DelegateProperty->GetName());
 	}
 
-	FObjectProperty* ComponentProperty = FindField<FObjectProperty>(Blueprint->SkeletonGeneratedClass, ThisComponentName);
+	FObjectProperty* ComponentProperty = FindFProperty<FObjectProperty>(Blueprint->SkeletonGeneratedClass, ThisComponentName);
 
 	if ( !ComponentProperty )
 	{
diff --git a/Engine/Source/Editor/UMGEditor/Private/Details/DetailWidgetExtensionHandler.cpp b/Engine/Source/Editor/UMGEditor/Private/Details/DetailWidgetExtensionHandler.cpp
--- a/Engine/Source/Editor/UMGEditor/Private/Details/DetailWidgetExtensionHandler.cpp
+++ b/Engine/Source/Editor/UMGEditor/Private/Details/DetailWidgetExtensionHandler.cpp
@@ -31,7 +31,7 @@ bool FDetailWidgetExtensionHandler::IsPropertyExtendable(const UClass* InObjectC
 
 		if ( UClass* ContainerClass = Property->GetOwner<UClass>() )
 		{
-			FDelegateProperty* DelegateProperty = FindField<FDelegateProperty>(ContainerClass, FName(*DelegateName));
+			FDelegateProperty* DelegateProperty = FindFProperty<FDelegateProperty>(ContainerClass, FName(*DelegateName));
 			if ( DelegateProperty )
 			{
 				return true;
diff --git a/Engine/Source/Editor/UMGEditor/Private/WidgetBlueprint.cpp b/Engine/Source/Editor/UMGEditor/Private/WidgetBlueprint.cpp
--- a/Engine/Source/Editor/UMGEditor/Private/WidgetBlueprint.cpp
+++ b/Engine/Source/Editor/UMGEditor/Private/WidgetBlueprint.cpp
@@ -203,11 +203,7 @@ FFieldVariant FEditorPropertyPathSegment::GetMember() const
 	FName FieldName = GetMemberName();
 	if ( FieldName != NAME_None )
 	{
-		FFieldVariant Field = FindField<UField>(Struct, FieldName);
-		if (!Field.IsValid())
-		{
-			Field = FindField<FField>(Struct, FieldName);
-		}
+		FFieldVariant Field = FindUFieldOrFProperty(Struct, FieldName);
 		//if ( Field == nullptr )
 		//{
 		//	if ( UClass* Class = Cast<UClass>(Struct) )
@@ -216,7 +212,7 @@ FFieldVariant FEditorPropertyPathSegment::GetMember() const
 		//		{
 		//			if ( UClass* SkeletonClass = Blueprint->SkeletonGeneratedClass )
 		//			{
-		//				Field = FindField<UField>(SkeletonClass, FieldName);
+		//				Field = FindUField<UField>(SkeletonClass, FieldName);
 		//			}
 		//		}
 		//	}
@@ -414,7 +410,7 @@ bool FDelegateEditorBinding::IsAttributePropertyBinding(UWidgetBlueprint* Bluepr
 	{
 		// Next find the underlying delegate we're actually binding to, if it's an event the name will be the same,
 		// for properties we need to lookup the delegate property we're actually going to be binding to.
-		FDelegateProperty* BindableProperty = FindField<FDelegateProperty>(TargetWidget->GetClass(), FName(*(PropertyName.ToString() + TEXT("Delegate"))));
+		FDelegateProperty* BindableProperty = FindFProperty<FDelegateProperty>(TargetWidget->GetClass(), FName(*(PropertyName.ToString() + TEXT("Delegate"))));
 		return BindableProperty != nullptr;
 	}
 
@@ -441,8 +437,8 @@ bool FDelegateEditorBinding::IsBindingValid(UClass* BlueprintGeneratedClass, UWi
 	{
 		// Next find the underlying delegate we're actually binding to, if it's an event the name will be the same,
 		// for properties we need to lookup the delegate property we're actually going to be binding to.
-		FDelegateProperty* BindableProperty = FindField<FDelegateProperty>(TargetWidget->GetClass(), FName(*( PropertyName.ToString() + TEXT("Delegate") )));
-		FDelegateProperty* EventProperty = FindField<FDelegateProperty>(TargetWidget->GetClass(), PropertyName);
+		FDelegateProperty* BindableProperty = FindFProperty<FDelegateProperty>(TargetWidget->GetClass(), FName(*( PropertyName.ToString() + TEXT("Delegate") )));
+		FDelegateProperty* EventProperty = FindFProperty<FDelegateProperty>(TargetWidget->GetClass(), PropertyName);
 
 		bool bNeedsToBePure = BindableProperty ? true : false;
 		FDelegateProperty* DelegateProperty = BindableProperty ? BindableProperty : EventProperty;
diff --git a/Engine/Source/Editor/UMGEditor/Private/WidgetBlueprintCompiler.cpp b/Engine/Source/Editor/UMGEditor/Private/WidgetBlueprintCompiler.cpp
--- a/Engine/Source/Editor/UMGEditor/Private/WidgetBlueprintCompiler.cpp
+++ b/Engine/Source/Editor/UMGEditor/Private/WidgetBlueprintCompiler.cpp
@@ -120,7 +120,7 @@ void FWidgetBlueprintCompilerContext::CreateFunctionList()
 		{
 			const FName PropertyName = EditorBinding.SourceProperty;
 
-			FProperty* Property = FindField<FProperty>(Blueprint->SkeletonGeneratedClass, PropertyName);
+			FProperty* Property = FindFProperty<FProperty>(Blueprint->SkeletonGeneratedClass, PropertyName);
 			if ( Property )
 			{
 				// Create the function graph.
diff --git a/Engine/Source/Editor/UMGEditor/Private/WidgetBlueprintEditorUtils.cpp b/Engine/Source/Editor/UMGEditor/Private/WidgetBlueprintEditorUtils.cpp
--- a/Engine/Source/Editor/UMGEditor/Private/WidgetBlueprintEditorUtils.cpp
+++ b/Engine/Source/Editor/UMGEditor/Private/WidgetBlueprintEditorUtils.cpp
@@ -1515,7 +1515,7 @@ void FWidgetBlueprintEditorUtils::ImportPropertiesFromText(UObject* Object, cons
 	{
 		for ( const auto& Entry : ExportedProperties )
 		{
-			if ( FProperty* Property = FindField<FProperty>(Object->GetClass(), Entry.Key) )
+			if ( FProperty* Property = FindFProperty<FProperty>(Object->GetClass(), Entry.Key) )
 			{
 				FEditPropertyChain PropertyChain;
 				PropertyChain.AddHead(Property);
diff --git a/Engine/Source/Editor/UnrealEd/Private/AssetSelection.cpp b/Engine/Source/Editor/UnrealEd/Private/AssetSelection.cpp
--- a/Engine/Source/Editor/UnrealEd/Private/AssetSelection.cpp
+++ b/Engine/Source/Editor/UnrealEd/Private/AssetSelection.cpp
@@ -912,7 +912,7 @@ bool FActorFactoryAssetProxy::ApplyMaterialToActor( AActor* TargetActor, UMateri
 		ALandscapeProxy* Landscape = Cast<ALandscapeProxy>(TargetActor);
 		if (Landscape != NULL)
 		{
-			FProperty* MaterialProperty = FindField<FProperty>(ALandscapeProxy::StaticClass(), "LandscapeMaterial");
+			FProperty* MaterialProperty = FindFProperty<FProperty>(ALandscapeProxy::StaticClass(), "LandscapeMaterial");
 			Landscape->PreEditChange(MaterialProperty);
 			Landscape->LandscapeMaterial = MaterialToApply;
 			FPropertyChangedEvent PropertyChangedEvent(MaterialProperty);
diff --git a/Engine/Source/Editor/UnrealEd/Private/ComponentVisualizer.cpp b/Engine/Source/Editor/UnrealEd/Private/ComponentVisualizer.cpp
--- a/Engine/Source/Editor/UnrealEd/Private/ComponentVisualizer.cpp
+++ b/Engine/Source/Editor/UnrealEd/Private/ComponentVisualizer.cpp
@@ -121,7 +121,7 @@ UActorComponent* FComponentPropertyPath::GetComponent() const
 
 				UClass* ActorClass = Actor->GetClass();
 				check(ActorClass);
-				FProperty* Prop = FindField<FProperty>(ActorClass, PropertyName);
+				FProperty* Prop = FindFProperty<FProperty>(ActorClass, PropertyName);
 
 				if (FObjectProperty* ObjectProp = CastField<FObjectProperty>(Prop))
 				{
@@ -207,7 +207,7 @@ UActorComponent* FComponentVisualizer::GetComponentFromPropertyName(const AActor
 	if(CompOwner && Property.IsValid())
 	{
 		UClass* ActorClass = CompOwner->GetClass();
-		FProperty* Prop = FindField<FProperty>(ActorClass, Property.Name);
+		FProperty* Prop = FindFProperty<FProperty>(ActorClass, Property.Name);
 		if (FObjectProperty* ObjectProp = CastField<FObjectProperty>(Prop))
 		{
 			UObject* Object = ObjectProp->GetObjectPropertyValue(ObjectProp->ContainerPtrToValuePtr<void>(CompOwner, Property.Index));
diff --git a/Engine/Source/Editor/UnrealEd/Private/EdMode.cpp b/Engine/Source/Editor/UnrealEd/Private/EdMode.cpp
--- a/Engine/Source/Editor/UnrealEd/Private/EdMode.cpp
+++ b/Engine/Source/Editor/UnrealEd/Private/EdMode.cpp
@@ -73,7 +73,7 @@ namespace
 			}
 
 			// Obtain the property info from the given structure definition
-			FProperty* CurrentProp = FindField<FProperty>(InStruct, FName(*NameToken));
+			FProperty* CurrentProp = FindFProperty<FProperty>(InStruct, FName(*NameToken));
 
 			// Check first to see if this is a simple structure (i.e. not an array of structures)
 			FStructProperty* StructProp = CastField<FStructProperty>(CurrentProp);
@@ -104,7 +104,7 @@ namespace
 		}
 		else
 		{
-			FProperty* Prop = FindField<FProperty>(InStruct, FName(*PropertyName));
+			FProperty* Prop = FindFProperty<FProperty>(InStruct, FName(*PropertyName));
 			if(Prop != NULL)
 			{
 				if( FArrayProperty* ArrayProp = CastField<FArrayProperty>(Prop) )
diff --git a/Engine/Source/Editor/UnrealEd/Private/EditorEngine.cpp b/Engine/Source/Editor/UnrealEd/Private/EditorEngine.cpp
--- a/Engine/Source/Editor/UnrealEd/Private/EditorEngine.cpp
+++ b/Engine/Source/Editor/UnrealEd/Private/EditorEngine.cpp
@@ -219,6 +219,7 @@
 #include "Subsystems/AssetEditorSubsystem.h"
 #include "StudioAnalytics.h"
 #include "Engine/LevelScriptActor.h"
+#include "UObject/UnrealType.h"
 
 #if WITH_CHAOS
 #include "ChaosSolversModule.h"
@@ -3465,7 +3466,7 @@ struct FConvertStaticMeshActorInfo
 
 	bool PropsDiffer(const TCHAR* PropertyPath, UObject* Obj)
 	{
-		const FProperty* PartsProp = FindField<FProperty>( PropertyPath );
+		const FProperty* PartsProp = FindFProperty<FProperty>( PropertyPath );
 		check(PartsProp);
 
 		uint8* ClassDefaults = (uint8*)Obj->GetClass()->GetDefaultObject();
diff --git a/Engine/Source/Editor/UnrealEd/Private/Kismet2/BlueprintEditorUtils.cpp b/Engine/Source/Editor/UnrealEd/Private/Kismet2/BlueprintEditorUtils.cpp
--- a/Engine/Source/Editor/UnrealEd/Private/Kismet2/BlueprintEditorUtils.cpp
+++ b/Engine/Source/Editor/UnrealEd/Private/Kismet2/BlueprintEditorUtils.cpp
@@ -2194,7 +2194,8 @@ bool FBlueprintEditorUtils::IsGraphNameUnique(UBlueprint* Blueprint, const FName
 	if( !FindObject<UObject>(Blueprint, *InName.ToString()) )
 	{
 		// Next, check for functions with that name in the blueprint's class scope
-		if( !FindField<UField>(Blueprint->SkeletonGeneratedClass, InName) )
+		FFieldVariant ExistingField = FindUFieldOrFProperty(Blueprint->SkeletonGeneratedClass, InName);
+		if( !ExistingField )
 		{
 			// Finally, check function entry points
 			TArray<UK2Node_Event*> AllEvents;
@@ -2334,7 +2335,7 @@ UClass* const FBlueprintEditorUtils::GetOverrideFunctionClass(UBlueprint* Bluepr
 
 	if (OverrideFunc == nullptr)
 	{
-		OverrideFunc = FindField<UFunction>(Blueprint->SkeletonGeneratedClass, FuncName);
+		OverrideFunc = FindUField<UFunction>(Blueprint->SkeletonGeneratedClass, FuncName);
 		// search up the class hierarchy, we want to find the original declaration of the function to match FBlueprintEventNodeSpawner.
 		// Doing so ensures that we can find the existing node if there is one:
 		const UClass* Iter = Blueprint->SkeletonGeneratedClass->GetSuperClass();
@@ -3770,12 +3771,12 @@ void FBlueprintEditorUtils::SetBlueprintVariableMetaData(UBlueprint* Blueprint,
 		else
 		{
 			Blueprint->NewVariables[VarIndex].SetMetaData(MetaDataKey, MetaDataValue);
-			FProperty* Property = FindField<FProperty>(Blueprint->SkeletonGeneratedClass, VarName);
+			FProperty* Property = FindFProperty<FProperty>(Blueprint->SkeletonGeneratedClass, VarName);
 			if (Property)
 			{
 				Property->SetMetaData(MetaDataKey, *MetaDataValue);
 			}
-			Property = FindField<FProperty>(Blueprint->GeneratedClass, VarName);
+			Property = FindFProperty<FProperty>(Blueprint->GeneratedClass, VarName);
 			if (Property)
 			{
 				Property->SetMetaData(MetaDataKey, *MetaDataValue);
@@ -3894,12 +3895,12 @@ void FBlueprintEditorUtils::RemoveBlueprintVariableMetaData(UBlueprint* Blueprin
 		else
 		{
 			Blueprint->NewVariables[VarIndex].RemoveMetaData(MetaDataKey);
-			FProperty* Property = FindField<FProperty>(Blueprint->SkeletonGeneratedClass, VarName);
+			FProperty* Property = FindFProperty<FProperty>(Blueprint->SkeletonGeneratedClass, VarName);
 			if (Property)
 			{
 				Property->RemoveMetaData(MetaDataKey);
 			}
-			Property = FindField<FProperty>(Blueprint->GeneratedClass, VarName);
+			Property = FindFProperty<FProperty>(Blueprint->GeneratedClass, VarName);
 			if (Property)
 			{
 				Property->RemoveMetaData(MetaDataKey);
@@ -3936,7 +3937,7 @@ void FBlueprintEditorUtils::SetBlueprintVariableCategory(UBlueprint* Blueprint,
 	Blueprint->Modify();
 
 	UClass* SkeletonGeneratedClass = Blueprint->SkeletonGeneratedClass;
-	if (FProperty* TargetProperty = FindField<FProperty>(SkeletonGeneratedClass, VarName))
+	if (FProperty* TargetProperty = FindFProperty<FProperty>(SkeletonGeneratedClass, VarName))
 	{
 		UClass* OuterClass = TargetProperty->GetOwnerChecked<UClass>();
 		const bool bIsNativeVar = (OuterClass->ClassGeneratedBy == nullptr);
@@ -4152,7 +4153,7 @@ FText FBlueprintEditorUtils::GetBlueprintVariableCategory(UBlueprint* Blueprint,
 {
 	FText CategoryName;
 	UClass* SkeletonGeneratedClass = Blueprint->SkeletonGeneratedClass;
-	FProperty* TargetProperty = FindField<FProperty>(SkeletonGeneratedClass, VarName);
+	FProperty* TargetProperty = FindFProperty<FProperty>(SkeletonGeneratedClass, VarName);
 	if(TargetProperty != nullptr)
 	{
 		CategoryName = FObjectEditorUtils::GetCategoryText(TargetProperty);
@@ -4693,7 +4694,7 @@ void FBlueprintEditorUtils::RenameMemberVariable(UBlueprint* Blueprint, const FN
 				UObject* GeneratedCDO = GeneratedClass ? GeneratedClass->GetDefaultObject(false) : nullptr;
 				if (GeneratedCDO)
 				{
-					FProperty* TargetProperty = FindField<FProperty>(GeneratedCDO->GetClass(), OldName); // GeneratedCDO->GetClass() is used instead of GeneratedClass, because CDO could use REINST class.
+					FProperty* TargetProperty = FindFProperty<FProperty>(GeneratedCDO->GetClass(), OldName); // GeneratedCDO->GetClass() is used instead of GeneratedClass, because CDO could use REINST class.
 					// Grab the address of where the property is actually stored (UObject* base, plus the offset defined in the property)
 					void* OldPropertyAddr = TargetProperty ? TargetProperty->ContainerPtrToValuePtr<void>(GeneratedCDO) : nullptr;
 					if (OldPropertyAddr)
@@ -5012,7 +5013,7 @@ FName FBlueprintEditorUtils::DuplicateVariable(UBlueprint* InBlueprint, const US
 			//Grab property of blueprint's current CDO
 			UClass* GeneratedClass = InBlueprint->GeneratedClass;
 			UObject* GeneratedCDO = GeneratedClass->GetDefaultObject();
-			FProperty* TargetProperty = FindField<FProperty>(GeneratedClass, Variable.VarName);
+			FProperty* TargetProperty = FindFProperty<FProperty>(GeneratedClass, Variable.VarName);
 
 			if( TargetProperty )
 			{
@@ -5238,8 +5239,8 @@ void FBlueprintEditorUtils::RenameLocalVariable(UBlueprint* InBlueprint, const U
 	{
 		UK2Node_FunctionEntry* FunctionEntry = nullptr;
 		FBPVariableDescription* LocalVariable = FindLocalVariable(InBlueprint, InScope, InOldName, &FunctionEntry);
-		const FProperty* OldProperty = FindField<const FProperty>(InScope, InOldName);
-		const FProperty* ExistingProperty = FindField<const FProperty>(InScope, InNewName);
+		const FProperty* OldProperty = FindFProperty<const FProperty>(InScope, InOldName);
+		const FProperty* ExistingProperty = FindFProperty<const FProperty>(InScope, InNewName);
 		const bool bHasExistingProperty = ExistingProperty && ExistingProperty != OldProperty;
 		if (bHasExistingProperty)
 		{
@@ -6028,7 +6029,7 @@ UFunction* FBlueprintEditorUtils::GetInterfaceFunction(UBlueprint* Blueprint, co
 	{
 		if (I.Interface)
 		{
-			Function = FindField<UFunction>(I.Interface, FuncName);
+			Function = FindUField<UFunction>(I.Interface, FuncName);
 			if (Function)
 			{
 				// found it, done
@@ -6042,7 +6043,7 @@ UFunction* FBlueprintEditorUtils::GetInterfaceFunction(UBlueprint* Blueprint, co
 	{
 		for (const FImplementedInterface& I : TempClass->Interfaces)
 		{
-			Function = FindField<UFunction>(I.Class, FuncName);
+			Function = FindUField<UFunction>(I.Class, FuncName);
 			if (Function)
 			{
 				// found it, done
@@ -6823,7 +6824,7 @@ static void ConformInterfaceByName(UBlueprint* Blueprint, FBPInterfaceDescriptio
 					}
 					// We perform the check here to avoid creating a graph if it isnt implemented in the full interface (note not the skeleton interface that we are iterating over)
 					// this is to avoid creating it then removing the graph below if it isnt present in the full class, which will cause a name conflict second time around
-					else if(FindField<UFunction>(CurrentInterfaceDesc.Interface, FunctionName))
+					else if(FindUField<UFunction>(CurrentInterfaceDesc.Interface, FunctionName))
 					{
 						UEdGraph* NewGraph = FBlueprintEditorUtils::CreateNewGraph(Blueprint, FunctionName, UAnimationGraph::StaticClass(), UAnimationGraphSchema::StaticClass());
 						NewGraph->bAllowDeletion = false;
@@ -6853,7 +6854,7 @@ static void ConformInterfaceByName(UBlueprint* Blueprint, FBPInterfaceDescriptio
 			// If we can't find the function associated with the graph, delete it
 			const UEdGraph* CurrentGraph = CurrentInterfaceDesc.Graphs[GraphIndex];
 
-			if (!CurrentGraph || !FindField<UFunction>(CurrentInterfaceDesc.Interface, CurrentGraph->GetFName()))
+			if (!CurrentGraph || !FindUField<UFunction>(CurrentInterfaceDesc.Interface, CurrentGraph->GetFName()))
 			{
 				CurrentInterfaceDesc.Graphs.RemoveAt(GraphIndex, 1);
 				GraphIndex--;
@@ -8740,7 +8741,7 @@ bool FBlueprintEditorUtils::IsPaletteActionReadOnly(TSharedPtr<FEdGraphSchemaAct
 					if(GraphAction->GraphType == EEdGraphSchemaAction_K2Graph::Function)
 					{
 						// Check if the function is an override
-						UFunction* OverrideFunc = FindField<UFunction>(BlueprintObj->ParentClass, GraphAction->FuncName);
+						UFunction* OverrideFunc = FindUField<UFunction>(BlueprintObj->ParentClass, GraphAction->FuncName);
 						if ( OverrideFunc != nullptr )
 						{
 							bIsReadOnly = true;
diff --git a/Engine/Source/Editor/UnrealEd/Private/Kismet2/ComponentEditorUtils.cpp b/Engine/Source/Editor/UnrealEd/Private/Kismet2/ComponentEditorUtils.cpp
--- a/Engine/Source/Editor/UnrealEd/Private/Kismet2/ComponentEditorUtils.cpp
+++ b/Engine/Source/Editor/UnrealEd/Private/Kismet2/ComponentEditorUtils.cpp
@@ -997,7 +997,7 @@ FName FComponentEditorUtils::FindVariableNameGivenComponentInstance(const UActor
 	if (AActor* OwnerActor = ComponentInstance->GetOwner())
 	{
 		UClass* OwnerActorClass = OwnerActor->GetClass();
-		if (FObjectProperty* TestProperty = FindField<FObjectProperty>(OwnerActorClass, ComponentInstance->GetFName()))
+		if (FObjectProperty* TestProperty = FindFProperty<FObjectProperty>(OwnerActorClass, ComponentInstance->GetFName()))
 		{
 			if (ComponentInstance->GetClass()->IsChildOf(TestProperty->PropertyClass))
 					{
diff --git a/Engine/Source/Editor/UnrealEd/Private/Kismet2/Kismet2.cpp b/Engine/Source/Editor/UnrealEd/Private/Kismet2/Kismet2.cpp
--- a/Engine/Source/Editor/UnrealEd/Private/Kismet2/Kismet2.cpp
+++ b/Engine/Source/Editor/UnrealEd/Private/Kismet2/Kismet2.cpp
@@ -2029,7 +2029,7 @@ void FKismetEditorUtilities::CreateNewBoundEventForActor(AActor* Actor, FName Ev
 	if ((Actor != nullptr) && (EventName != NAME_None))
 	{
 		// First, find the property we want to bind to
-		if (FMulticastDelegateProperty* DelegateProperty = FindField<FMulticastDelegateProperty>(Actor->GetClass(), EventName))
+		if (FMulticastDelegateProperty* DelegateProperty = FindFProperty<FMulticastDelegateProperty>(Actor->GetClass(), EventName))
 		{
 			// Get the correct level script blueprint
 			if (ULevelScriptBlueprint* LSB = Actor->GetLevel()->GetLevelScriptBlueprint())
@@ -2073,7 +2073,7 @@ void FKismetEditorUtilities::CreateNewBoundEventForClass(UClass* Class, FName Ev
 	if ( ( Class != nullptr ) && ( EventName != NAME_None ) && ( Blueprint != nullptr ) && ( ComponentProperty != nullptr ) )
 	{
 		// First, find the property we want to bind to
-		FMulticastDelegateProperty* DelegateProperty = FindField<FMulticastDelegateProperty>(Class, EventName);
+		FMulticastDelegateProperty* DelegateProperty = FindFProperty<FMulticastDelegateProperty>(Class, EventName);
 		if ( DelegateProperty != nullptr )
 		{
 			UEdGraph* TargetGraph = Blueprint->GetLastEditedUberGraph();
diff --git a/Engine/Source/Editor/UnrealEd/Private/Kismet2/Kismet2NameValidators.cpp b/Engine/Source/Editor/UnrealEd/Private/Kismet2/Kismet2NameValidators.cpp
--- a/Engine/Source/Editor/UnrealEd/Private/Kismet2/Kismet2NameValidators.cpp
+++ b/Engine/Source/Editor/UnrealEd/Private/Kismet2/Kismet2NameValidators.cpp
@@ -183,7 +183,7 @@ EValidatorResult FKismetNameValidator::IsValid(const FName& Name, bool /* bOrigi
 			}
 			else
 			{
-				if(FindField<const FProperty>(Scope, *Name.ToString()) != NULL)
+				if(FindFProperty<const FProperty>(Scope, *Name.ToString()) != NULL)
 				{
 					ValidatorResult = EValidatorResult::LocallyInUse;
 				}
diff --git a/Engine/Source/Editor/UnrealEd/Private/Kismet2/KismetReinstanceUtilities.cpp b/Engine/Source/Editor/UnrealEd/Private/Kismet2/KismetReinstanceUtilities.cpp
--- a/Engine/Source/Editor/UnrealEd/Private/Kismet2/KismetReinstanceUtilities.cpp
+++ b/Engine/Source/Editor/UnrealEd/Private/Kismet2/KismetReinstanceUtilities.cpp
@@ -394,7 +394,7 @@ void FBlueprintCompileReinstancer::GenerateFieldMappings(TMap<FFieldVariant, FFi
 
 	for (TPair<FName, FProperty*>& Prop : PropertyMap)
 	{
-		FieldMapping.Add(Prop.Value, FindField<FProperty>(ClassToReinstance, *Prop.Key.ToString()));
+		FieldMapping.Add(Prop.Value, FindFProperty<FProperty>(ClassToReinstance, *Prop.Key.ToString()));
 	}
 
 	for (TPair<FName, UFunction*>& Func : FunctionMap)
diff --git a/Engine/Source/Editor/UnrealEd/Private/Kismet2/StructureEditorUtils.cpp b/Engine/Source/Editor/UnrealEd/Private/Kismet2/StructureEditorUtils.cpp
--- a/Engine/Source/Editor/UnrealEd/Private/Kismet2/StructureEditorUtils.cpp
+++ b/Engine/Source/Editor/UnrealEd/Private/Kismet2/StructureEditorUtils.cpp
@@ -444,7 +444,7 @@ bool FStructureEditorUtils::ChangeVariableDefaultValue(UUserDefinedStruct* Struc
 		bool bAdvancedValidation = true;
 		if (!NewDefaultValue.IsEmpty())
 		{
-			const FProperty* Property = FindField<FProperty>(Struct, VarDesc->VarName);
+			const FProperty* Property = FindFProperty<FProperty>(Struct, VarDesc->VarName);
 			FStructOnScope StructDefaultMem(Struct);
 			bAdvancedValidation = StructDefaultMem.IsValid() && Property &&
 				FBlueprintEditorUtils::PropertyValueFromString(Property, NewDefaultValue, StructDefaultMem.GetStructMemory(), Struct);
@@ -513,7 +513,7 @@ FProperty* FStructureEditorUtils::GetPropertyByFriendlyName(const UUserDefinedSt
 		{
 			if (VarDesc.FriendlyName == DisplayName)
 			{
-				return FindField<FProperty>(Struct, VarDesc.VarName);
+				return FindFProperty<FProperty>(Struct, VarDesc.VarName);
 			}
 		}
 	}
@@ -709,7 +709,7 @@ bool FStructureEditorUtils::ChangeVariableTooltip(UUserDefinedStruct* Struct, FG
 		ModifyStructData(Struct);
 		VarDesc->ToolTip = InTooltip;
 
-		FProperty* Property = FindField<FProperty>(Struct, VarDesc->VarName);
+		FProperty* Property = FindFProperty<FProperty>(Struct, VarDesc->VarName);
 		if (Property)
 		{
 			Property->SetMetaData(FBlueprintMetadata::MD_Tooltip, *VarDesc->ToolTip);
@@ -792,7 +792,7 @@ bool FStructureEditorUtils::CanEnableMultiLineText(const UUserDefinedStruct* Str
 	const FStructVariableDescription* VarDesc = GetVarDescByGuid(Struct, VarGuid);
 	if (VarDesc)
 	{
-		FProperty* Property = FindField<FProperty>(Struct, VarDesc->VarName);
+		FProperty* Property = FindFProperty<FProperty>(Struct, VarDesc->VarName);
 
 		// If this is an array, we need to test its inner property as that's the real type
 		if (FArrayProperty* ArrayProperty = CastField<FArrayProperty>(Property))
@@ -819,7 +819,7 @@ bool FStructureEditorUtils::ChangeMultiLineTextEnabled(UUserDefinedStruct* Struc
 		ModifyStructData(Struct);
 
 		VarDesc->bEnableMultiLineText = bIsEnabled;
-		FProperty* Property = FindField<FProperty>(Struct, VarDesc->VarName);
+		FProperty* Property = FindFProperty<FProperty>(Struct, VarDesc->VarName);
 		if (Property)
 		{
 			if (VarDesc->bEnableMultiLineText)
@@ -869,7 +869,7 @@ bool FStructureEditorUtils::Change3dWidgetEnabled(UUserDefinedStruct* Struct, FG
 		ModifyStructData(Struct);
 
 		VarDesc->bEnable3dWidget = bIsEnabled;
-		FProperty* Property = FindField<FProperty>(Struct, VarDesc->VarName);
+		FProperty* Property = FindFProperty<FProperty>(Struct, VarDesc->VarName);
 		if (Property)
 		{
 			if (VarDesc->bEnable3dWidget)
@@ -903,7 +903,7 @@ FGuid FStructureEditorUtils::GetGuidForProperty(const FProperty* Property)
 FProperty* FStructureEditorUtils::GetPropertyByGuid(const UUserDefinedStruct* Struct, const FGuid VarGuid)
 {
 	const FStructVariableDescription* VarDesc = GetVarDescByGuid(Struct, VarGuid);
-	return VarDesc ? FindField<FProperty>(Struct, VarDesc->VarName) : nullptr;
+	return VarDesc ? FindFProperty<FProperty>(Struct, VarDesc->VarName) : nullptr;
 }
 
 FGuid FStructureEditorUtils::GetGuidFromPropertyName(const FName Name)
diff --git a/Engine/Source/Editor/UnrealEd/Private/LevelEditorViewport.cpp b/Engine/Source/Editor/UnrealEd/Private/LevelEditorViewport.cpp
--- a/Engine/Source/Editor/UnrealEd/Private/LevelEditorViewport.cpp
+++ b/Engine/Source/Editor/UnrealEd/Private/LevelEditorViewport.cpp
@@ -236,19 +236,19 @@ namespace LevelEditorViewportClientHelper
 		switch (WidgetMode)
 		{
 		case FWidget::WM_Translate:
-			ValueProperty = FindField<FProperty>(USceneComponent::StaticClass(), USceneComponent::GetRelativeLocationPropertyName());
+			ValueProperty = FindFProperty<FProperty>(USceneComponent::StaticClass(), USceneComponent::GetRelativeLocationPropertyName());
 			break;
 		case FWidget::WM_Rotate:
-			ValueProperty = FindField<FProperty>(USceneComponent::StaticClass(), USceneComponent::GetRelativeRotationPropertyName());
+			ValueProperty = FindFProperty<FProperty>(USceneComponent::StaticClass(), USceneComponent::GetRelativeRotationPropertyName());
 			break;
 		case FWidget::WM_Scale:
-			ValueProperty = FindField<FProperty>(USceneComponent::StaticClass(), USceneComponent::GetRelativeScale3DPropertyName());
+			ValueProperty = FindFProperty<FProperty>(USceneComponent::StaticClass(), USceneComponent::GetRelativeScale3DPropertyName());
 			break;
 		case FWidget::WM_TranslateRotateZ:
-			ValueProperty = FindField<FProperty>(USceneComponent::StaticClass(), USceneComponent::GetRelativeLocationPropertyName());
+			ValueProperty = FindFProperty<FProperty>(USceneComponent::StaticClass(), USceneComponent::GetRelativeLocationPropertyName());
 			break;
 		case FWidget::WM_2D:
-			ValueProperty = FindField<FProperty>(USceneComponent::StaticClass(), USceneComponent::GetRelativeLocationPropertyName());
+			ValueProperty = FindFProperty<FProperty>(USceneComponent::StaticClass(), USceneComponent::GetRelativeLocationPropertyName());
 			break;
 		default:
 			break;
diff --git a/Engine/Source/Editor/UnrealEd/Private/MaterialGraphNode_Comment.cpp b/Engine/Source/Editor/UnrealEd/Private/MaterialGraphNode_Comment.cpp
--- a/Engine/Source/Editor/UnrealEd/Private/MaterialGraphNode_Comment.cpp
+++ b/Engine/Source/Editor/UnrealEd/Private/MaterialGraphNode_Comment.cpp
@@ -106,7 +106,7 @@ void UMaterialGraphNode_Comment::PostPlacedNewNode()
 void UMaterialGraphNode_Comment::OnRenameNode(const FString& NewName)
 {
 	// send property changed events
-	FProperty* NodeCommentProperty = FindField<FProperty>(GetClass(), "NodeComment");
+	FProperty* NodeCommentProperty = FindFProperty<FProperty>(GetClass(), "NodeComment");
 	if(NodeCommentProperty != NULL)
 	{
 		PreEditChange(NodeCommentProperty);
diff --git a/Engine/Source/Editor/UnrealEd/Private/SSocketManager.cpp b/Engine/Source/Editor/UnrealEd/Private/SSocketManager.cpp
--- a/Engine/Source/Editor/UnrealEd/Private/SSocketManager.cpp
+++ b/Engine/Source/Editor/UnrealEd/Private/SSocketManager.cpp
@@ -140,7 +140,7 @@ private:
 			{
 				FScopedTransaction Transaction( LOCTEXT("SetSocketName", "Set Socket Name") );
 				
-				FProperty* ChangedProperty = FindField<FProperty>( UStaticMeshSocket::StaticClass(), "SocketName" );
+				FProperty* ChangedProperty = FindFProperty<FProperty>( UStaticMeshSocket::StaticClass(), "SocketName" );
 				
 				// Pre edit, calls modify on the object
 				SelectedSocket->PreEditChange(ChangedProperty);
diff --git a/Engine/Source/Editor/UnrealEd/Private/Tests/AutomationEditorCommon.cpp b/Engine/Source/Editor/UnrealEd/Private/Tests/AutomationEditorCommon.cpp
--- a/Engine/Source/Editor/UnrealEd/Private/Tests/AutomationEditorCommon.cpp
+++ b/Engine/Source/Editor/UnrealEd/Private/Tests/AutomationEditorCommon.cpp
@@ -232,7 +232,7 @@ void FAutomationEditorCommonUtils::ApplyCustomFactorySetting(UObject* InObject,
 	const FString PropertyName = PropertyChain[0];
 	PropertyChain.RemoveAt(0);
 
-	FProperty* TargetProperty = FindField<FProperty>(InObject->GetClass(), *PropertyName);
+	FProperty* TargetProperty = FindFProperty<FProperty>(InObject->GetClass(), *PropertyName);
 	if (TargetProperty)
 	{
 		if (PropertyChain.Num() == 0)
diff --git a/Engine/Source/Editor/UnrealEd/Private/Tests/AutomationEditorPromotionCommon.cpp b/Engine/Source/Editor/UnrealEd/Private/Tests/AutomationEditorPromotionCommon.cpp
--- a/Engine/Source/Editor/UnrealEd/Private/Tests/AutomationEditorPromotionCommon.cpp
+++ b/Engine/Source/Editor/UnrealEd/Private/Tests/AutomationEditorPromotionCommon.cpp
@@ -193,7 +193,7 @@ void FEditorPromotionTestUtilities::SendCommandToCurrentEditor(const FInputChord
 */
 FString FEditorPromotionTestUtilities::GetPropertyByName(UObject* TargetObject, const FString& InVariableName)
 {
-	FProperty* FoundProperty = FindField<FProperty>(TargetObject->GetClass(), *InVariableName);
+	FProperty* FoundProperty = FindFProperty<FProperty>(TargetObject->GetClass(), *InVariableName);
 	if (FoundProperty)
 	{
 		FString ValueString;
@@ -213,7 +213,7 @@ FString FEditorPromotionTestUtilities::GetPropertyByName(UObject* TargetObject,
 */
 void FEditorPromotionTestUtilities::SetPropertyByName(UObject* TargetObject, const FString& InVariableName, const FString& NewValueString)
 {
-	FProperty* FoundProperty = FindField<FProperty>(TargetObject->GetClass(), *InVariableName);
+	FProperty* FoundProperty = FindFProperty<FProperty>(TargetObject->GetClass(), *InVariableName);
 	if (FoundProperty)
 	{
 		const FScopedTransaction PropertyChanged(LOCTEXT("PropertyChanged", "Object Property Change"));
diff --git a/Engine/Source/Editor/ViewportInteraction/ActorViewportTransformable.cpp b/Engine/Source/Editor/ViewportInteraction/ActorViewportTransformable.cpp
--- a/Engine/Source/Editor/ViewportInteraction/ActorViewportTransformable.cpp
+++ b/Engine/Source/Editor/ViewportInteraction/ActorViewportTransformable.cpp
@@ -31,16 +31,16 @@ void FActorViewportTransformable::ApplyTransform( const FTransform& NewTransform
 				ExistingTransform.GetScale3D() == NewTransform.GetScale3D();
 			GEditor->BroadcastBeginObjectMovement(*Actor);
 			
-			FProperty* TransformProperty = FindField<FProperty>(USceneComponent::StaticClass(), USceneComponent::GetRelativeLocationPropertyName());
+			FProperty* TransformProperty = FindFProperty<FProperty>(USceneComponent::StaticClass(), USceneComponent::GetRelativeLocationPropertyName());
 			if (!bOnlyTranslationChanged)
 			{
 				if (ExistingTransform.GetRotation() != NewTransform.GetRotation())
 				{
-					TransformProperty = FindField<FProperty>(USceneComponent::StaticClass(), USceneComponent::GetRelativeRotationPropertyName());
+					TransformProperty = FindFProperty<FProperty>(USceneComponent::StaticClass(), USceneComponent::GetRelativeRotationPropertyName());
 				}
 				else if (ExistingTransform.GetScale3D() != NewTransform.GetScale3D())
 				{
-					TransformProperty = FindField<FProperty>(USceneComponent::StaticClass(), USceneComponent::GetRelativeScale3DPropertyName());
+					TransformProperty = FindFProperty<FProperty>(USceneComponent::StaticClass(), USceneComponent::GetRelativeScale3DPropertyName());
 				}
 			}
 			
diff --git a/Engine/Source/Editor/WorldBrowser/Private/StreamingLevels/StreamingLevelModel.cpp b/Engine/Source/Editor/WorldBrowser/Private/StreamingLevels/StreamingLevelModel.cpp
--- a/Engine/Source/Editor/WorldBrowser/Private/StreamingLevels/StreamingLevelModel.cpp
+++ b/Engine/Source/Editor/WorldBrowser/Private/StreamingLevels/StreamingLevelModel.cpp
@@ -96,7 +96,7 @@ void FStreamingLevelModel::SetLevelColor(FLinearColor InColor)
 {
 	if (LevelStreaming.IsValid())
 	{
-		FProperty* DrawColorProperty = FindField<FProperty>(LevelStreaming->GetClass(), "LevelColor");
+		FProperty* DrawColorProperty = FindFProperty<FProperty>(LevelStreaming->GetClass(), "LevelColor");
 		LevelStreaming->PreEditChange(DrawColorProperty);
 
 		LevelStreaming.Get()->LevelColor = InColor;
diff --git a/Engine/Source/Programs/UnrealHeaderTool/Private/HeaderParser.cpp b/Engine/Source/Programs/UnrealHeaderTool/Private/HeaderParser.cpp
--- a/Engine/Source/Programs/UnrealHeaderTool/Private/HeaderParser.cpp
+++ b/Engine/Source/Programs/UnrealHeaderTool/Private/HeaderParser.cpp
@@ -7616,7 +7616,7 @@ void FHeaderParser::CompileFunctionDeclaration(FClasses& AllClasses)
 	}
 	if (SuperStruct)
 	{
-		if (UFunction* OverriddenFunction = ::FindField<UFunction>(SuperStruct, FuncInfo.Function.Identifier))
+		if (UFunction* OverriddenFunction = ::FindUField<UFunction>(SuperStruct, FuncInfo.Function.Identifier))
 		{
 			// Native function overrides should be done in CPP text, not in a UFUNCTION() declaration (you can't change flags, and it'd otherwise be a burden to keep them identical)
 			UE_LOG_ERROR_UHT(TEXT("%s: Override of UFUNCTION in parent class (%s) cannot have a UFUNCTION() declaration above it; it will use the same parameters as the original declaration."), FuncInfo.Function.Identifier, *OverriddenFunction->GetOuter()->GetName());
diff --git a/Engine/Source/Runtime/CoreUObject/Private/UObject/Class.cpp b/Engine/Source/Runtime/CoreUObject/Private/UObject/Class.cpp
--- a/Engine/Source/Runtime/CoreUObject/Private/UObject/Class.cpp
+++ b/Engine/Source/Runtime/CoreUObject/Private/UObject/Class.cpp
@@ -1978,8 +1978,15 @@ void UStruct::SetSuperStruct(UStruct* NewSuperStruct)
 
 FString UStruct::PropertyNameToDisplayName(FName InName) const
 {
-	const UField* FoundField = FindField<const UField>(this, InName);
-	return GetAuthoredNameForField(FoundField);
+	FFieldVariant FoundField = FindUFieldOrFProperty(this, InName);
+	if (FoundField.IsUObject())
+	{
+		return GetAuthoredNameForField(FoundField.Get<UField>());
+	}
+	else
+	{
+		return GetAuthoredNameForField(FoundField.Get<FField>());
+	}
 }
 
 FString UStruct::GetAuthoredNameForField(const UField* Field) const
diff --git a/Engine/Source/Runtime/CoreUObject/Private/UObject/Field.cpp b/Engine/Source/Runtime/CoreUObject/Private/UObject/Field.cpp
--- a/Engine/Source/Runtime/CoreUObject/Private/UObject/Field.cpp
+++ b/Engine/Source/Runtime/CoreUObject/Private/UObject/Field.cpp
@@ -1124,7 +1124,7 @@ FString GetPathNameSafe(const FField* InField)
 	}
 }
 
-FField* FindFieldByPath(const TCHAR* InFieldPath)
+FField* FindFPropertyByPath(const TCHAR* InFieldPath)
 { 
 	// Expected format: FullPackageName.OwnerName:Field
 	FField* FoundField = nullptr;
@@ -1155,7 +1155,7 @@ FField* FindFieldByPath(const TCHAR* InFieldPath)
 				checkf(*TestChar != ':' && *TestChar != '.', TEXT("FindFieldByPath can't resolve nested properties: %s"), InFieldPath);
 			}
 #endif
-			FoundField = FindField<FField>(Owner, InFieldPath);
+			FoundField = FindFProperty<FField>(Owner, InFieldPath);
 		}
 	}
 	return FoundField;
diff --git a/Engine/Source/Runtime/CoreUObject/Private/UObject/FieldPath.cpp b/Engine/Source/Runtime/CoreUObject/Private/UObject/FieldPath.cpp
--- a/Engine/Source/Runtime/CoreUObject/Private/UObject/FieldPath.cpp
+++ b/Engine/Source/Runtime/CoreUObject/Private/UObject/FieldPath.cpp
@@ -123,7 +123,7 @@ FField* FFieldPath::TryToResolvePath(UStruct* InCurrentStruct, int32* OutOwnerIn
 	if (UStruct* Owner = Cast<UStruct>(LastOuter))
 	{
 		check(PathIndex <= 1);
-		Result = FindField<FField>(Owner, Path[PathIndex]);
+		Result = FindFProperty<FField>(Owner, Path[PathIndex]);
 		if (Result && PathIndex > 0)
 		{
 			// Nested property
diff --git a/Engine/Source/Runtime/CoreUObject/Private/UObject/Obj.cpp b/Engine/Source/Runtime/CoreUObject/Private/UObject/Obj.cpp
--- a/Engine/Source/Runtime/CoreUObject/Private/UObject/Obj.cpp
+++ b/Engine/Source/Runtime/CoreUObject/Private/UObject/Obj.cpp
@@ -3089,7 +3089,7 @@ static void PerformSetCommand( const TCHAR* Str, FOutputDevice& Ar, bool bNotify
 		UClass* Class = FindObject<UClass>(ANY_PACKAGE, ObjectName);
 		if (Class != NULL)
 		{
-			FProperty* Property = FindField<FProperty>(Class, PropertyName);
+			FProperty* Property = FindFProperty<FProperty>(Class, PropertyName);
 			if (Property != NULL)
 			{
 				while (*Str == ' ')
@@ -3108,7 +3108,7 @@ static void PerformSetCommand( const TCHAR* Str, FOutputDevice& Ar, bool bNotify
 			UObject* Object = FindObject<UObject>(ANY_PACKAGE, ObjectName);
 			if (Object != NULL)
 			{
-				FProperty* Property = FindField<FProperty>(Object->GetClass(), PropertyName);
+				FProperty* Property = FindFProperty<FProperty>(Object->GetClass(), PropertyName);
 				if (Property != NULL)
 				{
 					while (*Str == ' ')
@@ -3298,7 +3298,7 @@ bool StaticExec( UWorld* InWorld, const TCHAR* Cmd, FOutputDevice& Ar )
 		{
 			if
 			(	FParse::Token( Str, PropertyName, UE_ARRAY_COUNT(PropertyName), 1 )
-			&&	(Property=FindField<FProperty>( Class, PropertyName))!=NULL )
+			&&	(Property=FindFProperty<FProperty>( Class, PropertyName))!=NULL )
 			{
 				FString	Temp;
 				if( Class->GetDefaultsCount() )
@@ -3459,7 +3459,7 @@ bool StaticExec( UWorld* InWorld, const TCHAR* Cmd, FOutputDevice& Ar )
 		{
 			FParse::Token(Str,PropertyName,UE_ARRAY_COUNT(PropertyName),1);
 			{
-				Property=FindField<FProperty>(Class,PropertyName);
+				Property=FindFProperty<FProperty>(Class,PropertyName);
 				{
 					int32 cnt = 0;
 					UObject* LimitOuter = NULL;
@@ -3600,7 +3600,7 @@ bool StaticExec( UWorld* InWorld, const TCHAR* Cmd, FOutputDevice& Ar )
 
 		if (FParse::Token(Str, ClassName, UE_ARRAY_COUNT(ClassName), true))
 		{
-			//if ( (Property=FindField<FProperty>(Class,PropertyName)) != NULL )
+			//if ( (Property=FindFProperty<FProperty>(Class,PropertyName)) != NULL )
 			UClass* Class = FindObject<UClass>(ANY_PACKAGE, ClassName);
 
 			if (Class != NULL)
@@ -3631,7 +3631,7 @@ bool StaticExec( UWorld* InWorld, const TCHAR* Cmd, FOutputDevice& Ar )
 
 			if (Class != NULL)
 			{
-				UFunction* Function = FindField<UFunction>(Class, FunctionName);
+				UFunction* Function = FindUField<UFunction>(Class, FunctionName);
 
 				if (Function != NULL)
 				{
diff --git a/Engine/Source/Runtime/CoreUObject/Private/UObject/Property.cpp b/Engine/Source/Runtime/CoreUObject/Private/UObject/Property.cpp
--- a/Engine/Source/Runtime/CoreUObject/Private/UObject/Property.cpp
+++ b/Engine/Source/Runtime/CoreUObject/Private/UObject/Property.cpp
@@ -1116,7 +1116,7 @@ const TCHAR* FProperty::ImportSingleProperty( const TCHAR* Str, void* DestData,
 		}
 
 		const FName PropertyName(Len, Start);
-		FProperty* Property = FindField<FProperty>(ObjectStruct, PropertyName);
+		FProperty* Property = FindFProperty<FProperty>(ObjectStruct, PropertyName);
 
 		if (Property == nullptr)
 		{
@@ -1125,7 +1125,7 @@ const TCHAR* FProperty::ImportSingleProperty( const TCHAR* Str, void* DestData,
 
 			if (NewPropertyName != NAME_None)
 			{
-				Property = FindField<FProperty>(ObjectStruct, NewPropertyName);
+				Property = FindFProperty<FProperty>(ObjectStruct, NewPropertyName);
 			}
 
 			if (!Property)
diff --git a/Engine/Source/Runtime/CoreUObject/Private/UObject/PropertyHelper.h b/Engine/Source/Runtime/CoreUObject/Private/UObject/PropertyHelper.h
--- a/Engine/Source/Runtime/CoreUObject/Private/UObject/PropertyHelper.h
+++ b/Engine/Source/Runtime/CoreUObject/Private/UObject/PropertyHelper.h
@@ -255,7 +255,7 @@ namespace DelegatePropertyTools
 				}
 			}
 		}
-		UFunction *Func = FindField<UFunction>( Cls, FuncName );
+		UFunction *Func = FindUField<UFunction>( Cls, FuncName );
 		// Check function params.
 		if(Func != nullptr)
 		{
diff --git a/Engine/Source/Runtime/CoreUObject/Private/UObject/PropertyNumeric.cpp b/Engine/Source/Runtime/CoreUObject/Private/UObject/PropertyNumeric.cpp
--- a/Engine/Source/Runtime/CoreUObject/Private/UObject/PropertyNumeric.cpp
+++ b/Engine/Source/Runtime/CoreUObject/Private/UObject/PropertyNumeric.cpp
@@ -24,7 +24,7 @@ int64 FNumericProperty::ReadEnumAsInt64(FStructuredArchive::FSlot Slot, UStruct*
 	FName EnumName;
 	Slot << EnumName;
 
-	UEnum* Enum = FindField<UEnum>(dynamic_cast<UClass*>(DefaultsStruct) ? static_cast<UClass*>(DefaultsStruct) : DefaultsStruct->GetTypedOuter<UClass>(), Tag.EnumName);
+	UEnum* Enum = FindUField<UEnum>(dynamic_cast<UClass*>(DefaultsStruct) ? static_cast<UClass*>(DefaultsStruct) : DefaultsStruct->GetTypedOuter<UClass>(), Tag.EnumName);
 	if (!Enum)
 	{
 		Enum = FindObject<UEnum>(ANY_PACKAGE, *Tag.EnumName.ToString());
