Removed literal reference to UE4

#rnx

diff --git a/Engine/Source/Editor/TranslationEditor/Private/TranslationDataManager.cpp b/Engine/Source/Editor/TranslationEditor/Private/TranslationDataManager.cpp
--- a/Engine/Source/Editor/TranslationEditor/Private/TranslationDataManager.cpp
+++ b/Engine/Source/Editor/TranslationEditor/Private/TranslationDataManager.cpp
@@ -455,7 +455,7 @@ void FTranslationDataManager::GetHistoryForTranslationUnits()
 				FString SavedDir;	// Store these cached translation history files in the saved directory
 				if (IsEngineManifest)
 				{
-					ProjectName = "Engine";
+					ProjectName = TEXT("Engine");
 					SavedDir = FPaths::EngineSavedDir();
 				}
 				else
@@ -464,8 +464,7 @@ void FTranslationDataManager::GetHistoryForTranslationUnits()
 					SavedDir = FPaths::ProjectSavedDir();
 				}
 
-				FString TempFileName = SavedDir / "CachedTranslationHistory" / "UE4-Manifest-" + ProjectName + "-" + FPaths::GetBaseFilename(InManifestFilePath) + "-Rev-" + FString::FromInt(Revision->GetRevisionNumber());
-
+				FString TempFileName = SavedDir / TEXT("CachedTranslationHistory") / FString::Printf(TEXT("Manifest-%s-%s-Rev-%d"), *ProjectName, *FPaths::GetBaseFilename(InManifestFilePath), Revision->GetRevisionNumber());
 				
 				if (!FPaths::FileExists(TempFileName))	// Don't bother syncing again if we already have this manifest version cached locally
 				{
@@ -483,7 +482,7 @@ void FTranslationDataManager::GetHistoryForTranslationUnits()
 						{
 							for (FTranslationContextInfo& ContextInfo : TranslationUnit->Contexts)
 							{
-								FString PreviousSourceText = "";
+								FString PreviousSourceText;
 
 								// If we already have history, then compare against the newest history so far
 								if (ContextInfo.Changes.Num() > 0)
