Fixed the Translation Picker reporting everything as untranslated

#rnx

diff --git a/Engine/Source/Editor/TranslationEditor/Private/TranslationPickerEditWindow.cpp b/Engine/Source/Editor/TranslationEditor/Private/TranslationPickerEditWindow.cpp
--- a/Engine/Source/Editor/TranslationEditor/Private/TranslationPickerEditWindow.cpp
+++ b/Engine/Source/Editor/TranslationEditor/Private/TranslationPickerEditWindow.cpp
@@ -236,18 +236,16 @@ void STranslationPickerEditWidget::Construct(const FArguments& InArgs)
 	int32 DefaultPadding = 0.0f;
 
 	// Try and get the localization information for this text
-	FString Namespace;
-	FString Key;
+	FTextId TextId;
 	FString SourceString;
 	FString TranslationString;
-	bool bHasIdentity = false;
 	{
 		if (const FString* SourceStringPtr = FTextInspector::GetSourceString(PickedText))
 		{
 			SourceString = *SourceStringPtr;
 		}
 		TranslationString = FTextInspector::GetDisplayString(PickedText);
-		bHasIdentity = !FTextInspector::GetTextId(PickedText).IsEmpty();
+		TextId = FTextInspector::GetTextId(PickedText);
 	}
 
 	// Try and find the LocRes the active translation came from
@@ -255,7 +253,7 @@ void STranslationPickerEditWidget::Construct(const FArguments& InArgs)
 	FString LocResPath;
 	FString LocTargetName;
 	FString LocResCultureName;
-	if (bHasIdentity && FTextLocalizationManager::Get().GetLocResID(Namespace, Key, LocResPath))
+	if (!TextId.IsEmpty() && FTextLocalizationManager::Get().GetLocResID(TextId.GetNamespace(), TextId.GetKey(), LocResPath))
 	{
 		LocTargetName = FPaths::GetBaseFilename(LocResPath);
 
@@ -264,18 +262,18 @@ void STranslationPickerEditWidget::Construct(const FArguments& InArgs)
 	}
 
 	// Clean the package localization ID from the namespace (to mirror what the text gatherer does when scraping for translation data)
-	Namespace = TextNamespaceUtil::StripPackageNamespace(Namespace);
+	FString CleanNamespace = TextNamespaceUtil::StripPackageNamespace(TextId.GetNamespace().GetChars());
 
 	// Save the necessary data in UTranslationUnit for later.  This is what we pass to TranslationDataManager to save our edits
 	TranslationUnit = NewObject<UTranslationUnit>();
-	TranslationUnit->Namespace = Namespace;
-	TranslationUnit->Key = Key;
+	TranslationUnit->Namespace = CleanNamespace;
+	TranslationUnit->Key = TextId.GetKey().GetChars();
 	TranslationUnit->Source = SourceString;
 	TranslationUnit->Translation = TranslationString;
 	TranslationUnit->LocresPath = LocResPath;
 
 	// Can only save if we have have an identity and are in a known localization target file
-	bHasRequiredLocalizationInfoForSaving = bHasIdentity && !LocTargetName.IsEmpty();
+	bHasRequiredLocalizationInfoForSaving = !TextId.IsEmpty() && !LocTargetName.IsEmpty();
 
 	TSharedPtr<SGridPanel> GridPanel;
 
@@ -342,7 +340,7 @@ void STranslationPickerEditWidget::Construct(const FArguments& InArgs)
 		]
 	];
 
-	if (bHasIdentity)
+	if (!TextId.IsEmpty())
 	{
 		GridPanel->AddSlot(0, 2)
 			.Padding(FMargin(2.5))
@@ -357,7 +355,7 @@ void STranslationPickerEditWidget::Construct(const FArguments& InArgs)
 			[
 				SNew(SEditableTextBox)
 				.IsReadOnly(true)
-				.Text(FText::AsCultureInvariant(Namespace))
+				.Text(FText::AsCultureInvariant(CleanNamespace))
 			];
 		GridPanel->AddSlot(0, 3)
 			.Padding(FMargin(2.5))
@@ -372,7 +370,7 @@ void STranslationPickerEditWidget::Construct(const FArguments& InArgs)
 			[
 				SNew(SEditableTextBox)
 				.IsReadOnly(true)
-				.Text(FText::AsCultureInvariant(Key))
+				.Text(FText::AsCultureInvariant(TextId.GetKey().GetChars()))
 			];
 		
 		if (bHasRequiredLocalizationInfoForSaving)
