PR #9249: Added element count to translation editor tabs (Contributed by [REDACTED_PERSON])

diff --git a/Engine/Source/Editor/TranslationEditor/Private/TranslationEditor.cpp b/Engine/Source/Editor/TranslationEditor/Private/TranslationEditor.cpp
--- a/Engine/Source/Editor/TranslationEditor/Private/TranslationEditor.cpp
+++ b/Engine/Source/Editor/TranslationEditor/Private/TranslationEditor.cpp
@@ -239,6 +239,13 @@ FLinearColor FTranslationEditor::GetWorldCentricTabColorScale() const
 {
 	return FLinearColor( 0.0f, 0.0f, 0.2f, 0.5f );
 }
+FText FTranslationEditor::GetTabLabel_Untranslated() const
+{
+	const int32 UntranslatedNum = DataManager->GetUntranslatedArray().Num();
+	return UntranslatedNum > 0 ?
+		FText::Format(LOCTEXT("UntranslatedTabTitleNumbered", "Untranslated ({0})"), UntranslatedNum) :
+		LOCTEXT("UntranslatedTabTitle", "Untranslated");
+}
 
 TSharedRef<SDockTab> FTranslationEditor::SpawnTab_Untranslated( const FSpawnTabArgs& Args )
 {
@@ -280,9 +287,11 @@ TSharedRef<SDockTab> FTranslationEditor::SpawnTab_Untranslated( const FSpawnTabA
 	UntranslatedPropertyTableWidgetHandle = PropertyEditorModule.CreatePropertyTableWidgetHandle( UntranslatedPropertyTable.ToSharedRef(), CustomColumns);
 	TSharedRef<SWidget> PropertyTableWidget = UntranslatedPropertyTableWidgetHandle->GetWidget();
 
+	
+
 	TSharedRef<SDockTab> NewDockTab = SNew(SDockTab)
 		//.Icon( FAppStyle::GetBrush("TranslationEditor.Tabs.Properties") )
-		.Label( LOCTEXT("UntranslatedTabTitle", "Untranslated") )
+		.Label(this, &FTranslationEditor::GetTabLabel_Untranslated)
 		.TabColorScale( GetTabColorScale() )
 		[
 			SNew(SBorder)
@@ -298,6 +307,14 @@ TSharedRef<SDockTab> FTranslationEditor::SpawnTab_Untranslated( const FSpawnTabA
 	return NewDockTab;
 }
 
+FText FTranslationEditor::GetTabLabel_Review() const
+{
+	const int32 ReviewNum = DataManager->GetReviewArray().Num();
+	return ReviewNum > 0 ?
+		FText::Format(LOCTEXT("ReviewTabTitleNumbered", "Needs Review ({0})"), ReviewNum) :
+		LOCTEXT("ReviewTabTitle", "Needs Review");
+}
+
 TSharedRef<SDockTab> FTranslationEditor::SpawnTab_Review( const FSpawnTabArgs& Args )
 {
 	check( Args.GetTabId().TabType == ReviewTabId );
@@ -345,9 +362,11 @@ TSharedRef<SDockTab> FTranslationEditor::SpawnTab_Review( const FSpawnTabArgs& A
 	ReviewPropertyTableWidgetHandle = PropertyEditorModule.CreatePropertyTableWidgetHandle( ReviewPropertyTable.ToSharedRef(), CustomColumns);
 	TSharedRef<SWidget> PropertyTableWidget = ReviewPropertyTableWidgetHandle->GetWidget();
 
+	
+
 	TSharedRef<SDockTab> NewDockTab = SNew(SDockTab)
 		//.Icon( FAppStyle::GetBrush("TranslationEditor.Tabs.Properties") )
-		.Label( LOCTEXT("ReviewTabTitle", "Needs Review") )
+		.Label(this, &FTranslationEditor::GetTabLabel_Review)
 		.TabColorScale( GetTabColorScale() )
 		[
 			SNew(SBorder)
@@ -363,6 +382,14 @@ TSharedRef<SDockTab> FTranslationEditor::SpawnTab_Review( const FSpawnTabArgs& A
 	return NewDockTab;
 }
 
+FText FTranslationEditor::GetTabLabel_Completed() const
+{
+	const int32 CompletedNum = DataManager->GetCompleteArray().Num();
+	return CompletedNum > 0 ?
+		FText::Format(LOCTEXT("CompletedTabTitleNumbered", "Completed ({0})"), CompletedNum) :
+		LOCTEXT("CompletedTabTitle", "Completed");
+}
+
 TSharedRef<SDockTab> FTranslationEditor::SpawnTab_Completed( const FSpawnTabArgs& Args )
 {
 	check( Args.GetTabId().TabType == CompletedTabId );
@@ -403,9 +430,11 @@ TSharedRef<SDockTab> FTranslationEditor::SpawnTab_Completed( const FSpawnTabArgs
 	CompletedPropertyTableWidgetHandle = PropertyEditorModule.CreatePropertyTableWidgetHandle( CompletedPropertyTable.ToSharedRef(), CustomColumns);
 	TSharedRef<SWidget> PropertyTableWidget = CompletedPropertyTableWidgetHandle->GetWidget();
 
+	
+
 	TSharedRef<SDockTab> NewDockTab = SNew(SDockTab)
 		//.Icon( FAppStyle::GetBrush("TranslationEditor.Tabs.Properties") )
-		.Label( LOCTEXT("CompletedTabTitle", "Completed") )
+		.Label(this, &FTranslationEditor::GetTabLabel_Completed)
 		.TabColorScale( GetTabColorScale() )
 		[
 			SNew(SBorder)
diff --git a/Engine/Source/Editor/TranslationEditor/Private/TranslationEditor.h b/Engine/Source/Editor/TranslationEditor/Private/TranslationEditor.h
--- a/Engine/Source/Editor/TranslationEditor/Private/TranslationEditor.h
+++ b/Engine/Source/Editor/TranslationEditor/Private/TranslationEditor.h
@@ -106,10 +106,16 @@ private:
 
 	/**	Spawns the untranslated tab */
 	TSharedRef<SDockTab> SpawnTab_Untranslated( const FSpawnTabArgs& Args );
+	/** Getter function bound to the label attribute for the untranslated tab. */
+	FText GetTabLabel_Untranslated() const;
 	/**	Spawns the review tab */
 	TSharedRef<SDockTab> SpawnTab_Review( const FSpawnTabArgs& Args );
+	/** Getter function bound to the label attribute for the review tab. */
+	FText GetTabLabel_Review() const;
 	/**	Spawns the completed tab */
 	TSharedRef<SDockTab> SpawnTab_Completed( const FSpawnTabArgs& Args );
+	/** Getter function bound to the label attribute for the completed tab. */
+	FText GetTabLabel_Completed() const;
 	/**	Spawns the preview tab */
 	TSharedRef<SDockTab> SpawnTab_Preview( const FSpawnTabArgs& Args );
 	/**	Spawns the context tab */
